Introducing Capture Source Framework with some utility classes

diff --git a/Engine/Plugins/VirtualProduction/CaptureManager/Source/CaptureSourceFramework/CaptureSourceFramework.Build.cs b/Engine/Plugins/VirtualProduction/CaptureManager/Source/CaptureSourceFramework/CaptureSourceFramework.Build.cs
new file mode 100644
--- /dev/null
+++ b/Engine/Plugins/VirtualProduction/CaptureManager/Source/CaptureSourceFramework/CaptureSourceFramework.Build.cs
@@ -0,0 +1,31 @@
+// Copyright Epic Games, Inc. All Rights Reserved.
+
+using UnrealBuildTool;
+using System.IO;
+
+public class CaptureSourceFramework : ModuleRules
+{
+	public CaptureSourceFramework(ReadOnlyTargetRules Target) : base(Target)
+	{
+		PCHUsage = PCHUsageMode.UseExplicitOrSharedPCHs;
+
+		PublicDependencyModuleNames.AddRange(new string[]
+		{
+			"Core",
+			"CoreUObject",
+			"Engine",
+			"CaptureUtils"
+		});
+
+		PrivateDependencyModuleNames.AddRange(new string[]
+		{
+			"InputCore",
+			"AudioEditor",
+			"UnrealEd",
+			"DesktopWidgets",
+			"Slate",
+			"SlateCore",
+			"RenderCore",
+		});
+	}
+}
diff --git a/Engine/Plugins/VirtualProduction/CaptureManager/Source/CaptureSourceFramework/Private/Capabilities/DiskCapacityCapability.cpp b/Engine/Plugins/VirtualProduction/CaptureManager/Source/CaptureSourceFramework/Private/Capabilities/DiskCapacityCapability.cpp
new file mode 100644
--- /dev/null
+++ b/Engine/Plugins/VirtualProduction/CaptureManager/Source/CaptureSourceFramework/Private/Capabilities/DiskCapacityCapability.cpp
@@ -0,0 +1,73 @@
+// Copyright Epic Games, Inc. All Rights Reserved.
+
+#include "Capabilities/DiskCapacityCapability.h"
+
+const FString FStateUpdatedEvent::Name = TEXT("StateUpdated");
+
+const FString FUpdateCmd::Name = TEXT("Update");
+
+FUpdateCmd::FUpdateCmd()
+	: FCommandBase(Name)
+{
+}
+
+const FString FDiskCapacityCapabilityBase::Name = TEXT("DiskCapacity");
+const FString FDiskCapacityCapabilityBase::TotalSpace = TEXT("TotalSpace");
+const FString FDiskCapacityCapabilityBase::RemainingSpace = TEXT("RemainingSpace");
+const FString FDiskCapacityCapabilityBase::UsedSpace = TEXT("UsedSpace");
+
+FDiskCapacityCapabilityBase::FDiskCapacityCapabilityBase()
+	: FCaptureSourceCapability(Name)
+{
+	AddProperty(FPropertyDesc(TotalSpace, FPropertyDesc::EType::Number, FPropertyDesc::EAccess::ReadOnly));
+	AddProperty(FPropertyDesc(RemainingSpace, FPropertyDesc::EType::Number, FPropertyDesc::EAccess::ReadOnly));
+	AddProperty(FPropertyDesc(UsedSpace, FPropertyDesc::EType::Number, FPropertyDesc::EAccess::ReadOnly));
+
+	AddCommand(FCommandDesc(FUpdateCmd::Name, {}));
+
+	RegisterEvent(FStateUpdatedEvent::Name);
+}
+
+FDiskCapacityCapability::FDiskCapacityCapability()
+{
+}
+
+uint64 FDiskCapacityCapability::GetTotalSpace()
+{
+	return GetPropertyValue(TotalSpace).Get<int64>();
+}
+
+uint64 FDiskCapacityCapability::GetRemainingSpace()
+{
+	return GetPropertyValue(RemainingSpace).Get<int64>();
+}
+
+uint64 FDiskCapacityCapability::GetUsedSpace()
+{
+	return GetPropertyValue(TotalSpace).Get<int64>() - GetPropertyValue(RemainingSpace).Get<int64>();
+}
+
+void FDiskCapacityCapability::Update()
+{
+	ExecuteCommand(MakeShared<FUpdateCmd>());
+}
+
+void FDiskCapacityCapability::SetTotalSpacePropertyHandler(FPropertyGetter InGetter)
+{
+	SetPropertyGetter(TotalSpace, MoveTemp(InGetter));
+}
+
+void FDiskCapacityCapability::SetRemainingSpacePropertyHandler(FPropertyGetter InGetter)
+{
+	SetPropertyGetter(RemainingSpace, MoveTemp(InGetter));
+}
+
+void FDiskCapacityCapability::SetUpdateCmdHandler(FCommandHandler InCommandHandler)
+{
+	SetCommandHandler(FUpdateCmd::Name, MoveTemp(InCommandHandler));
+}
+
+void FDiskCapacityCapability::PublishStateUpdatedEvent(uint64 InTotalBytes, uint64 InRemainingBytes)
+{
+	PublishEvent<FStateUpdatedEvent>(InTotalBytes, InRemainingBytes);
+}
\ No newline at end of file
diff --git a/Engine/Plugins/VirtualProduction/CaptureManager/Source/CaptureSourceFramework/Private/Capabilities/IngestCapability.cpp b/Engine/Plugins/VirtualProduction/CaptureManager/Source/CaptureSourceFramework/Private/Capabilities/IngestCapability.cpp
new file mode 100644
--- /dev/null
+++ b/Engine/Plugins/VirtualProduction/CaptureManager/Source/CaptureSourceFramework/Private/Capabilities/IngestCapability.cpp
@@ -0,0 +1,143 @@
+// Copyright Epic Games, Inc. All Rights Reserved.
+
+#include "Capabilities/IngestCapability.h"
+
+FIngestCapabilityError::FIngestCapabilityError(ECode InCode, FString InMessage)
+	: Code(InCode)
+	, Message(MoveTemp(InMessage))
+{
+}
+
+FIngestCapabilityError::ECode FIngestCapabilityError::GetCode() const
+{
+	return Code;
+}
+
+const FString& FIngestCapabilityError::GetMessage() const
+{
+	return Message;
+}
+
+const FString FIngestCapabilityBase::Name = TEXT("Ingest");
+const FString FIngestCapabilityBase::CancelCmd = TEXT("Cancel");
+
+const FString FIngestCapabilityBase::Id = TEXT("Id");
+const FString FIngestCapabilityBase::Location = TEXT("Location");
+
+FIngestCapabilityBase::FIngestCapabilityBase()
+	: FCaptureSourceCapability(Name)
+{
+	AddCommand(FCommandDesc(FImportTakeCmd::Name,
+							{
+								FPropertyDesc(Id, FPropertyDesc::EType::Number),
+								FPropertyDesc(Location, FPropertyDesc::EType::String)
+							}));
+	AddCommand(FCommandDesc(FUpdateTakeListCmd::Name, {}));
+
+	AddCommand(FCommandDesc(CancelCmd,
+							{
+								FPropertyDesc(Id, FPropertyDesc::EType::Number)
+							}));
+}
+
+FIngestCapabilityBase::FTakeProcessHandle::FTakeProcessHandle(FIngestCapabilityBase* InOwner, FTakeId InId)
+	: Owner(InOwner)
+	, Id(InId)
+{
+}
+
+FTakeId FIngestCapabilityBase::AddTakeInfo(FTakeInfo InTakeInfo)
+{
+	FTakeId TakeId = CurrentTakeId++;
+	InTakeInfo.Id = TakeId;
+	Takes.Emplace(TakeId, MoveTemp(InTakeInfo));
+
+	return TakeId;
+}
+
+FTakeInfo FIngestCapabilityBase::GetTakeInfo(FTakeId InTakeId) const
+{
+	return Takes[InTakeId];
+}
+
+TArray<FTakeId> FIngestCapabilityBase::GetTakeIds() const
+{
+	TArray<FTakeId> TakeIds;
+	Takes.GetKeys(TakeIds);
+
+	return TakeIds;
+}
+
+void FIngestCapabilityBase::FTakeProcessHandle::Cancel()
+{
+	Owner->Cancel(Id);
+}
+
+const FString FImportTakeCmd::Name = TEXT("ImportTake");
+FImportTakeCmd::FImportTakeCmd(FProcessProgressReporter InProgressReporter,
+							   FProcessFinished InProcessFinished)
+	: FCommandBase(Name)
+	, ProgressReporter(MoveTemp(InProgressReporter))
+	, ProcessFinished(MoveTemp(InProcessFinished))
+{
+}
+
+const FString FUpdateTakeListCmd::Name = TEXT("UpdateTakeList");
+FUpdateTakeListCmd::FUpdateTakeListCmd(FCallback InCallback)
+	: FCommandBase(Name)
+	, Callback(MoveTemp(InCallback))
+{
+}
+
+FIngestCapability::FIngestCapability()
+{
+}
+
+FIngestCapability::FTakeProcessHandle FIngestCapability::ImportTake(FTakeId InTakeId,
+																	FString InLocation,
+																	FImportTakeCmd::FProcessFinished InProcessFinished,
+																	FImportTakeCmd::FProcessProgressReporter InProcessProgressReported)
+{
+	FTakeInfo& TakeInfo = Takes[InTakeId];
+
+	TSharedPtr<FImportTakeCmd> InCommand = MakeShared<FImportTakeCmd>(MoveTemp(InProcessProgressReported),
+																	  MoveTemp(InProcessFinished));
+
+	InCommand->AddParamValue(Id, MakePropertyValue<int64>(TakeInfo.Id));
+	InCommand->AddParamValue(Location, MakePropertyValue(MoveTemp(InLocation)));
+
+	ExecuteCommand(MoveTemp(InCommand));
+
+	return FTakeProcessHandle(this, InTakeId);
+}
+
+void FIngestCapability::UpdateTakeList(FUpdateTakeListCmd::FCallback InUpdateTakeListCallback)
+{
+	Takes.Empty();
+
+	ExecuteCommand(MakeShared<FUpdateTakeListCmd>(MoveTemp(InUpdateTakeListCallback)));
+}
+
+void FIngestCapability::Cancel(FTakeId InTakeId)
+{
+	TSharedPtr<FCommandBase> InCommand = MakeShared<FCommandBase>(CancelCmd);
+
+	InCommand->AddParamValue(Id, MakePropertyValue<int64>(InTakeId));
+
+	ExecuteCommand(MoveTemp(InCommand));
+}
+
+void FIngestCapability::SetImportTakeHandler(FCommandHandler InCommandHandler)
+{
+	SetCommandHandler(FImportTakeCmd::Name, MoveTemp(InCommandHandler));
+}
+
+void FIngestCapability::SetUpdateTakeListHandler(FCommandHandler InCommandHandler)
+{
+	SetCommandHandler(FUpdateTakeListCmd::Name, MoveTemp(InCommandHandler));
+}
+
+void FIngestCapability::SetCancelHandler(FCommandHandler InCommandHandler)
+{
+	SetCommandHandler(CancelCmd, MoveTemp(InCommandHandler));
+}
\ No newline at end of file
diff --git a/Engine/Plugins/VirtualProduction/CaptureManager/Source/CaptureSourceFramework/Private/Capabilities/RecordingCapability.cpp b/Engine/Plugins/VirtualProduction/CaptureManager/Source/CaptureSourceFramework/Private/Capabilities/RecordingCapability.cpp
new file mode 100644
--- /dev/null
+++ b/Engine/Plugins/VirtualProduction/CaptureManager/Source/CaptureSourceFramework/Private/Capabilities/RecordingCapability.cpp
@@ -0,0 +1,53 @@
+// Copyright Epic Games, Inc. All Rights Reserved.
+
+#include "Capabilities/RecordingCapability.h"
+
+const FString FRecordingEvent::Name = TEXT("IsRecording");
+
+const FString FRecordingCapabilityBase::Name = TEXT("Recording");
+const FString FRecordingCapabilityBase::SlateName = TEXT("SlateName");
+const FString FRecordingCapabilityBase::TakeNumber = TEXT("TakeNumber");
+const FString FRecordingCapabilityBase::StartRecordingCmd = TEXT("StartRecording");
+const FString FRecordingCapabilityBase::StopRecordingCmd = TEXT("StopRecording");
+
+FRecordingCapabilityBase::FRecordingCapabilityBase()
+	: FCaptureSourceCapability(Name)
+{
+	AddCommand(FCommandDesc(StartRecordingCmd, { FPropertyDesc(SlateName, FPropertyDesc::EType::String), FPropertyDesc(TakeNumber, FPropertyDesc::EType::Number) }));
+	AddCommand(FCommandDesc(StopRecordingCmd, {}));
+
+	RegisterEvent(FRecordingEvent::Name);
+}
+
+FRecordingCapability::FRecordingCapability()
+{
+}
+
+void FRecordingCapability::StartRecording(FString InSlateName, int64 InTakeNumber)
+{
+	TSharedPtr<FCommandBase> InCommand = MakeShared<FCommandBase>(StartRecordingCmd);
+	InCommand->AddParamValue(SlateName, MakePropertyValue<FString>(MoveTemp(InSlateName)));
+	InCommand->AddParamValue(TakeNumber, MakePropertyValue<int64>(InTakeNumber));
+
+	ExecuteCommand(MoveTemp(InCommand));
+}
+
+void FRecordingCapability::StopRecording()
+{
+	ExecuteCommand(MakeShared<FCommandBase>(StopRecordingCmd));
+}
+
+void FRecordingCapability::SetStartRecordingHandler(FCommandHandler InCommandHandler)
+{
+	SetCommandHandler(StartRecordingCmd, MoveTemp(InCommandHandler));
+}
+
+void FRecordingCapability::SetStopRecordingHandler(FCommandHandler InCommandHandler)
+{
+	SetCommandHandler(StopRecordingCmd, MoveTemp(InCommandHandler));
+}
+
+void FRecordingCapability::PublishRecordingEvent(bool bIsRecording)
+{
+	PublishEvent<FRecordingEvent>(bIsRecording);
+}
\ No newline at end of file
diff --git a/Engine/Plugins/VirtualProduction/CaptureManager/Source/CaptureSourceFramework/Private/CaptureSource.cpp b/Engine/Plugins/VirtualProduction/CaptureManager/Source/CaptureSourceFramework/Private/CaptureSource.cpp
new file mode 100644
--- /dev/null
+++ b/Engine/Plugins/VirtualProduction/CaptureManager/Source/CaptureSourceFramework/Private/CaptureSource.cpp
@@ -0,0 +1,58 @@
+// Copyright Epic Games, Inc. All Rights Reserved.
+
+#include "CaptureSource.h"
+
+const FString FReachableEvent::Name = "Reachable";
+
+FCaptureSource::FCaptureSource(const FString& InCaptureSourceFactoryId, const FString& InName)
+	: CaptureSourceFactoryId(InCaptureSourceFactoryId)
+	, Name(InName)
+{
+	// Register event
+	RegisterEvent(FReachableEvent::Name);
+}
+
+FCaptureSourceCapability* FCaptureSource::GetCapability(FString InName) const
+{
+	if (CapabilityMap.Contains(InName))
+	{
+		return CapabilityMap[InName];
+	}
+	
+	return nullptr;
+}
+
+TArray<FString> FCaptureSource::GetCapabilities() const
+{
+	TArray<FString> Capabilities;
+
+	CapabilityMap.GetKeys(Capabilities);
+
+	return Capabilities;
+}
+
+const FString& FCaptureSource::GetCaptureSourceFactoryId() const
+{
+	return CaptureSourceFactoryId;
+}
+
+const FString& FCaptureSource::GetName() const
+{
+	return Name;
+}
+
+void FCaptureSource::SetCreationParams(TMap<FString, FPropertyValue> InCreationParams)
+{
+	CreationParams = MoveTemp(InCreationParams);
+}
+
+const TMap<FString, FPropertyValue>& FCaptureSource::GetCreationParams() const
+{
+	return CreationParams;
+}
+
+void FCaptureSource::AddCapability(FCaptureSourceCapability* InCapability)
+{
+	FString CapabilityName = InCapability->GetName();
+	CapabilityMap.Emplace(MoveTemp(CapabilityName), InCapability);
+}
\ No newline at end of file
diff --git a/Engine/Plugins/VirtualProduction/CaptureManager/Source/CaptureSourceFramework/Private/CaptureSourceCapability.cpp b/Engine/Plugins/VirtualProduction/CaptureManager/Source/CaptureSourceFramework/Private/CaptureSourceCapability.cpp
new file mode 100644
--- /dev/null
+++ b/Engine/Plugins/VirtualProduction/CaptureManager/Source/CaptureSourceFramework/Private/CaptureSourceCapability.cpp
@@ -0,0 +1,129 @@
+// Copyright Epic Games, Inc. All Rights Reserved.
+
+#include "CaptureSourceCapability.h"
+
+FCaptureSourceCapabilityDesc::FCaptureSourceCapabilityDesc(FString InName, TArray<FPropertyDesc> InProperties, TArray<FCommandDesc> InCommands)
+	: Name(MoveTemp(InName))
+	, Properties(MoveTemp(InProperties))
+	, Commands(MoveTemp(InCommands))
+{
+}
+
+const FString& FCapturePropertyChangedEvent::EventName = TEXT("PropertyChanged");
+FCapturePropertyChangedEvent::FCapturePropertyChangedEvent(const FString& InName, FPropertyValue InValue)
+	: FCaptureEvent(EventName)
+	, Name(InName)
+	, Value(MoveTemp(InValue))
+{
+}
+
+FCaptureSourceCapability::FCaptureSourceCapability(FString InName)
+	: Name(MoveTemp(InName))
+{
+	RegisterEvent(FCapturePropertyChangedEvent::EventName);
+}
+
+void FCaptureSourceCapability::PublishPropertyChangedEvent(const FString& InName, FPropertyValue InValue)
+{
+	Properties.CheckPropertyValueTypeMatch(InName, InValue);
+
+	PublishEvent<FCapturePropertyChangedEvent>(InName, MoveTemp(InValue));
+}
+
+void FCaptureSourceCapability::SetPropertyValue(const FString& InName, FPropertyValue InValue)
+{
+	Properties.CheckPropertyAllowsWrite(InName);
+	Properties.CheckPropertyValueTypeMatch(InName, InValue);
+
+	if (Setters.Contains(InName))
+	{
+		FPropertySetter& Setter = Setters[InName];
+		if (Setter.IsBound())
+		{
+			Setter.Execute(InName, MoveTemp(InValue));
+		}
+	}
+}
+
+FPropertyValue FCaptureSourceCapability::GetPropertyValue(const FString& InName)
+{
+	Properties.CheckPropertyExists(InName);
+
+	if (Getters.Contains(InName))
+	{
+		FPropertyGetter& Getter = Getters[InName];
+		if (Getter.IsBound())
+		{
+			return Getter.Execute(InName);
+		}
+	}
+	
+	return FPropertyValue();
+}
+
+TArray<FPropertyDesc> FCaptureSourceCapability::GetProperties() const
+{
+	return Properties.GetProperties();
+}
+
+TArray<FCommandDesc> FCaptureSourceCapability::GetCommands() const
+{
+	return Commands.GetCommands();
+}
+
+FString FCaptureSourceCapability::GetName() const
+{
+	return Name;
+}
+
+FCaptureSourceCapabilityDesc FCaptureSourceCapability::CreateDescriptor()
+{
+	return FCaptureSourceCapabilityDesc(Name, Properties.GetProperties(), Commands.GetCommands());
+}
+
+void FCaptureSourceCapability::SetPropertySetter(const FString& InName, FPropertySetter InSetter)
+{
+	Properties.CheckPropertyAllowsWrite(InName);
+
+	Setters.Emplace(InName, MoveTemp(InSetter));
+}
+
+void FCaptureSourceCapability::SetPropertyGetter(const FString& InName, FPropertyGetter InGetter)
+{
+	Properties.CheckPropertyExists(InName);
+
+	Getters.Emplace(InName, MoveTemp(InGetter));
+}
+
+void FCaptureSourceCapability::SetCommandHandler(const FString& InName, FCommandHandler InCommandHandler)
+{
+	Commands.CheckCommandExists(InName);
+
+	CommandHandlers.Emplace(InName, MoveTemp(InCommandHandler));
+}
+
+void FCaptureSourceCapability::AddProperty(FPropertyDesc InPropertyDesc)
+{
+	Properties.AddProperty(MoveTemp(InPropertyDesc));
+}
+
+void FCaptureSourceCapability::AddCommand(FCommandDesc InCommand)
+{
+	Commands.AddCommand(MoveTemp(InCommand));
+}
+
+void FCaptureSourceCapability::ExecuteCommand(TSharedPtr<FCommandBase> InCommand)
+{
+	const FString& CommandName = InCommand->GetName();
+
+	Commands.CheckCommandParamValue(CommandName, InCommand);
+
+	if (CommandHandlers.Contains(CommandName))
+	{
+		FCommandHandler& Handler = CommandHandlers[CommandName];
+		if (Handler.IsBound())
+		{
+			Handler.Execute(MoveTemp(InCommand));
+		}
+	}
+}
\ No newline at end of file
diff --git a/Engine/Plugins/VirtualProduction/CaptureManager/Source/CaptureSourceFramework/Private/CaptureSourceCapabilityCommand.cpp b/Engine/Plugins/VirtualProduction/CaptureManager/Source/CaptureSourceFramework/Private/CaptureSourceCapabilityCommand.cpp
new file mode 100644
--- /dev/null
+++ b/Engine/Plugins/VirtualProduction/CaptureManager/Source/CaptureSourceFramework/Private/CaptureSourceCapabilityCommand.cpp
@@ -0,0 +1,90 @@
+// Copyright Epic Games, Inc. All Rights Reserved.
+
+#include "CaptureSourceCapabilityCommand.h"
+
+FCommandDesc::FCommandDesc(FString InName, TArray<FPropertyDesc> InParams)
+	: Name(MoveTemp(InName))
+	, Params(MoveTemp(InParams))
+{
+}
+
+void FCommandList::AddCommand(FCommandDesc InCommandDesc)
+{
+	Commands.Add(MoveTemp(InCommandDesc));
+}
+
+TArray<FCommandDesc> FCommandList::GetCommands() const
+{
+	return Commands;
+}
+
+void FCommandList::CheckCommandExists(const FString& InName)
+{
+	check(Commands.ContainsByPredicate([&InName](const FCommandDesc& InCommand)
+	{
+		return InCommand.Name == InName;
+	}));
+}
+
+void FCommandList::CheckCommandParamValue(const FString& InName, const TSharedPtr<FCommandBase>& InCommand)
+{
+	FCommandDesc* Command = Commands.FindByPredicate([&InName](const FCommandDesc& InCommand)
+	{
+		return InCommand.Name == InName;
+	});
+
+	check(Command);
+
+	for (const FPropertyDesc& Param : Command->Params)
+	{
+		check(InCommand->ContainsParamValue(Param.Name));
+	}
+
+	for (const TPair<FString, FPropertyValue>& ParamValue : InCommand->GetParamValues())
+	{
+		check(Command->Params.ContainsByPredicate([Name = ParamValue.Key](const FPropertyDesc& InElem)
+		{
+			return InElem.Name == Name;
+		}));
+	}
+
+	for (const FPropertyDesc& Param : Command->Params)
+	{
+		check(FPropertyList::CheckPropertyValueType(Param.Type, InCommand->GetParamValue(Param.Name)));
+	}
+}
+
+FCommandBase::FCommandBase(FString InName)
+	: Name(MoveTemp(InName))
+{
+}
+
+void FCommandBase::AddParamValue(const FString& InName, FPropertyValue InValue)
+{
+	Values.Emplace(InName, MoveTemp(InValue));
+}
+
+FPropertyValue FCommandBase::GetParamValue(const FString& InName) const
+{
+	return Values[InName];
+}
+
+void FCommandBase::SetParamValues(FValues InValues)
+{
+	Values = MoveTemp(InValues);
+}
+
+const FCommandBase::FValues& FCommandBase::GetParamValues() const
+{
+	return Values;
+}
+
+bool FCommandBase::ContainsParamValue(const FString& InName) const
+{
+	return Values.Contains(InName);
+}
+
+const FString& FCommandBase::GetName() const
+{
+	return Name;
+}
\ No newline at end of file
diff --git a/Engine/Plugins/VirtualProduction/CaptureManager/Source/CaptureSourceFramework/Private/CaptureSourceCapabilityProperty.cpp b/Engine/Plugins/VirtualProduction/CaptureManager/Source/CaptureSourceFramework/Private/CaptureSourceCapabilityProperty.cpp
new file mode 100644
--- /dev/null
+++ b/Engine/Plugins/VirtualProduction/CaptureManager/Source/CaptureSourceFramework/Private/CaptureSourceCapabilityProperty.cpp
@@ -0,0 +1,84 @@
+// Copyright Epic Games, Inc. All Rights Reserved.
+
+#include "CaptureSourceCapabilityProperty.h"
+
+FPropertyDesc::FPropertyDesc(FString InName, EType InType, EAccess InAccess)
+	: Name(MoveTemp(InName))
+	, Type(InType)
+	, Access(InAccess)
+{
+}
+
+FPropertyDesc::FPropertyDesc(FString InName, EType InType)
+	: Name(MoveTemp(InName))
+	, Type(InType)
+	, Access(EAccess::ReadWrite)
+{
+}
+
+void FPropertyList::AddProperty(FPropertyDesc InPropertyDesc)
+{
+	Properties.Add(MoveTemp(InPropertyDesc));
+}
+
+TArray<FPropertyDesc> FPropertyList::GetProperties() const
+{
+	return Properties;
+}
+
+void FPropertyList::CheckPropertyExists(const FString& InName)
+{
+	check(Properties.ContainsByPredicate([&InName](const FPropertyDesc& InProperty)
+	{
+		return InProperty.Name == InName;
+	}));
+}
+
+void FPropertyList::CheckPropertyValueTypeMatch(const FString& InName, const FPropertyValue& InValue)
+{
+	const FPropertyDesc* Property = Properties.FindByPredicate([&InName](const FPropertyDesc& InProperty)
+	{
+		return InProperty.Name == InName;
+	});
+
+	check(Property);
+
+	check(CheckPropertyValueType(Property->Type, InValue));
+}
+
+void FPropertyList::CheckPropertyAllowsWrite(const FString& InName)
+{
+	const FPropertyDesc* Property = Properties.FindByPredicate([&InName](const FPropertyDesc& InProperty)
+	{
+		return InProperty.Name == InName;
+	});
+
+	check(Property);
+
+	check(Property->Access == FPropertyDesc::EAccess::ReadWrite);
+}
+
+bool FPropertyList::CheckPropertyValueType(FPropertyDesc::EType InType, const FPropertyValue& InValue)
+{
+	switch (InType)
+	{
+		case FPropertyDesc::EType::Bool:
+			return InValue.IsType<bool>();
+		case FPropertyDesc::EType::Number:
+			return InValue.IsType<int64>();
+		case FPropertyDesc::EType::String:
+			return InValue.IsType<FString>();
+		case FPropertyDesc::EType::FloatingPoint:
+			return InValue.IsType<double>();
+		case FPropertyDesc::EType::ArrayBool:
+			return InValue.IsType<TArray<bool>>();
+		case FPropertyDesc::EType::ArrayNumber:
+			return InValue.IsType<TArray<int64>>();
+		case FPropertyDesc::EType::ArrayString:
+			return InValue.IsType<TArray<FString>>();
+		case FPropertyDesc::EType::ArrayFloatingPoint:
+			return InValue.IsType<TArray<double>>();
+		default:
+			return false;
+	}
+}
\ No newline at end of file
diff --git a/Engine/Plugins/VirtualProduction/CaptureManager/Source/CaptureSourceFramework/Private/CaptureSourceDetailsCustomization.cpp b/Engine/Plugins/VirtualProduction/CaptureManager/Source/CaptureSourceFramework/Private/CaptureSourceDetailsCustomization.cpp
new file mode 100644
--- /dev/null
+++ b/Engine/Plugins/VirtualProduction/CaptureManager/Source/CaptureSourceFramework/Private/CaptureSourceDetailsCustomization.cpp
@@ -0,0 +1,530 @@
+// Copyright Epic Games, Inc. All Rights Reserved.
+
+#include "CaptureSourceDetailsCustomization.h"
+
+#include "CaptureSourceFrameworkModule.h"
+
+#include "Widgets/Input/SButton.h"
+#include "Widgets/Input/SEditableTextBox.h"
+#include "Widgets/Input/SCheckBox.h"
+#include "Widgets/Input/SSpinBox.h"
+#include "Widgets/Input/SNumericEntryBox.h"
+
+static FString DefaultBase = TEXT("Default");
+
+FCaptureSourceDetailsCustomization::FCaptureSourceDetailsCustomization()
+{
+}
+
+void FCaptureSourceDetailsCustomization::RegisterPropertyCustomization(const FString& InCaptureSourceClassId,
+																	   const FString& InCapabilityName, 
+																	   const FString& InPropertyName, 
+																	   FPropertyWidgetCreator InWidgetCreator)
+{
+	PropertyCustomizations.Emplace(GenerateKey(InCaptureSourceClassId, InCapabilityName, InPropertyName), MoveTemp(InWidgetCreator));
+}
+
+void FCaptureSourceDetailsCustomization::RegisterCommandCustomization(const FString& InCaptureSourceClassId,
+																	  const FString& InCapabilityName,
+																	  const FString& InCommandName,
+																	  FCommandWidgetCreator InWidgetCreator)
+{
+	CommandCustomizations.Emplace(GenerateKey(InCaptureSourceClassId, InCapabilityName, InCommandName), MoveTemp(InWidgetCreator));
+}
+
+void FCaptureSourceDetailsCustomization::RegisterCapabilityCustomization(const FString& InCaptureSourceClassId,
+																		 const FString& InCapabilityName,
+																		 FCapabilityWidgetCreator InWidgetCreator)
+{
+	CapabilityCustomizations.Emplace(GenerateKey(InCaptureSourceClassId, InCapabilityName), MoveTemp(InWidgetCreator));
+}
+
+void FCaptureSourceDetailsCustomization::RegisterCaptureSourceCustomization(const FString& InCaptureSourceClassId,
+																			FCaptureSourceWidgetCreator InWidgetCreator)
+{
+	CreationParamsCustomizations.Emplace(GenerateKey(InCaptureSourceClassId), MoveTemp(InWidgetCreator));
+}
+
+TSharedPtr<SWidget> FCaptureSourceDetailsCustomization::CreatePropertyWidget(FCaptureSourceId InCaptureSourceId,
+																			 FCaptureSourceCapability* InCapability, 
+																			 const FPropertyDesc& InProperty)
+{
+	FCaptureSourceFrameworkModule& Module = FModuleManager::LoadModuleChecked<FCaptureSourceFrameworkModule>("CaptureSourceFramework");
+	const FString& InCaptureSourceFactory = Module.GetCaptureSourceManager().GetCaptureSourceFactoryById(InCaptureSourceId);
+
+	FPropertyWidgetCreator* FoundDelegate = PropertyCustomizations.Find(GenerateKey(InCaptureSourceFactory, InCapability->GetName(), InProperty.Name));
+	if (FoundDelegate && FoundDelegate->IsBound())
+	{
+		TSharedPtr<SWidget> Widget = FoundDelegate->Execute(InCaptureSourceId, InCapability, InProperty);
+		return Widget;
+	}
+
+	return CreateDefaultPropertyWidget(InCapability, InProperty);
+}
+
+TSharedPtr<SWidget> FCaptureSourceDetailsCustomization::CreateCommandWidget(FCaptureSourceId InCaptureSourceId,
+																			FCaptureSourceCapability* InCapability, 
+																			const FCommandDesc& InCommand)
+{
+	FCaptureSourceFrameworkModule& Module = FModuleManager::LoadModuleChecked<FCaptureSourceFrameworkModule>("CaptureSourceFramework");
+	const FString& InCaptureSourceFactory = Module.GetCaptureSourceManager().GetCaptureSourceFactoryById(InCaptureSourceId);
+
+	FCommandWidgetCreator* FoundDelegate = CommandCustomizations.Find(GenerateKey(InCaptureSourceFactory, InCapability->GetName(), InCommand.Name));
+	if (FoundDelegate && FoundDelegate->IsBound())
+	{
+		TSharedPtr<SWidget> Widget = FoundDelegate->Execute(InCaptureSourceId, InCapability, InCommand);
+		return Widget;
+	}
+
+	return CreateDefaultCommandWidget(InCapability, InCommand);
+}
+
+TSharedPtr<SWidget> FCaptureSourceDetailsCustomization::CreateCapabilityWidget(FCaptureSourceId InCaptureSourceId,
+																			   FCaptureSourceCapability* InCapability)
+{
+	FCaptureSourceFrameworkModule& Module = FModuleManager::LoadModuleChecked<FCaptureSourceFrameworkModule>("CaptureSourceFramework");
+	const FString& InCaptureSourceFactory = Module.GetCaptureSourceManager().GetCaptureSourceFactoryById(InCaptureSourceId);
+
+	FCapabilityWidgetCreator* FoundDelegate = CapabilityCustomizations.Find(GenerateKey(InCaptureSourceFactory, InCapability->GetName()));
+	if (FoundDelegate && FoundDelegate->IsBound())
+	{
+		TSharedPtr<SWidget> Widget = FoundDelegate->Execute(InCaptureSourceId, InCapability);
+		return Widget;
+	}
+
+	TSharedPtr<SVerticalBox> VerticalBox = SNew(SVerticalBox);
+
+	TMap<FString, TSharedPtr<FPropertyValue>> PropertyValueHolder;
+
+	for (const FPropertyDesc& Property : InCapability->GetProperties())
+	{
+		if (ContainsPropertyWidgetCustomization(InCaptureSourceFactory, InCapability->GetName(), Property.Name))
+		{
+			VerticalBox->AddSlot()
+				.AutoHeight()
+				.VAlign(EVerticalAlignment::VAlign_Center)
+				[
+					CreatePropertyWidget(InCaptureSourceId, InCapability, Property).ToSharedRef()
+				];
+		}
+		else
+		{
+			FPropertyValue Value = InCapability->GetPropertyValue(Property.Name);
+			if (Value.IsType<FEmptyVariantState>())
+			{
+				Value = GetDefaultValueForProperty(Property);
+			}
+
+			PropertyValueHolder.Emplace(Property.Name, MakeShared<FPropertyValue>(MoveTemp(Value)));
+
+			VerticalBox->AddSlot()
+				.AutoHeight()
+				.VAlign(EVerticalAlignment::VAlign_Center)
+				[
+					CreateDefaultCapabilityPropertyWidget(InCapability, Property, PropertyValueHolder[Property.Name]).ToSharedRef()
+				];
+		}
+	}
+
+	InCapability->SubscribeToEvent(
+		FCapturePropertyChangedEvent::EventName,
+		FCaptureEventHandler::Type::CreateRaw(this, &FCaptureSourceDetailsCustomization::HandlePropertyChangedEvent, MoveTemp(PropertyValueHolder)));
+
+	for (const FCommandDesc& Command : InCapability->GetCommands())
+	{
+		TSharedPtr<SVerticalBox> CommandVerticalBox = SNew(SVerticalBox);
+
+		CommandVerticalBox->AddSlot()
+			.AutoHeight()
+			.VAlign(EVerticalAlignment::VAlign_Center)
+			[
+				SNew(STextBlock)
+					.Text(FText::FromString(Command.Name))
+			];
+
+		CommandVerticalBox->AddSlot()
+			.AutoHeight()
+			.VAlign(EVerticalAlignment::VAlign_Center)
+			[
+				CreateCommandWidget(InCaptureSourceId, InCapability, Command).ToSharedRef()
+			];
+
+		VerticalBox->AddSlot()
+			.AutoHeight()
+			.VAlign(EVerticalAlignment::VAlign_Center)
+			[
+				CommandVerticalBox.ToSharedRef()
+			];
+	}
+
+	return VerticalBox;
+}
+
+TSharedPtr<SWindow> FCaptureSourceDetailsCustomization::CreateCaptureSourceWidget(const FCaptureSourceDescriptor& InCreationParams,
+																				  FCaptureSourceCreated InCaptureSourceCreated)
+{
+	FCaptureSourceWidgetCreator* FoundDelegate = CreationParamsCustomizations.Find(InCreationParams.GetId());
+	if (FoundDelegate && FoundDelegate->IsBound())
+	{
+		TSharedPtr<SWindow> Widget = FoundDelegate->Execute(InCreationParams, MoveTemp(InCaptureSourceCreated));
+		return Widget;
+	}
+
+	TSharedPtr<SWindow> Window = SNew(SWindow)
+		.Title(FText::FromString(FString::Format(TEXT("Creating {0}"), { InCreationParams.GetId() })))
+		.SizingRule(ESizingRule::Autosized);
+
+	TSharedPtr<SVerticalBox> VerticalBox = SNew(SVerticalBox);
+	TSharedPtr<TMap<FString, FPropertyValue>> ParamValues = MakeShared<TMap<FString, FPropertyValue>>();
+
+	// Name
+	TSharedPtr<FPropertyValue> NameValue = MakeShared<FPropertyValue>(MakePropertyValue<FString>(TEXT("Unnamed")));
+
+	TSharedRef<SWidget> NameWidget = 
+		CreateParamWidget(FPropertyDesc(TEXT("Name"), FPropertyDesc::EType::String), *NameValue).ToSharedRef();
+
+	VerticalBox->AddSlot()
+		.AutoHeight()
+		.VAlign(EVerticalAlignment::VAlign_Center)
+		.HAlign(EHorizontalAlignment::HAlign_Fill)
+		.AttachWidget(MoveTemp(NameWidget));
+
+	// Other creation params
+	for (const FPropertyDesc& Property : InCreationParams.GetCreationParams())
+	{
+		ParamValues->Emplace(Property.Name, GetDefaultValueForProperty(Property));
+
+		TSharedRef<SWidget> Widget = CreateParamWidget(Property, (*ParamValues)[Property.Name]).ToSharedRef();
+
+		VerticalBox->AddSlot()
+			.AutoHeight()
+			.VAlign(EVerticalAlignment::VAlign_Center)
+			.HAlign(EHorizontalAlignment::HAlign_Fill)
+			.AttachWidget(MoveTemp(Widget));
+	}
+
+	VerticalBox->AddSlot()
+		.AutoHeight()
+		.VAlign(EVerticalAlignment::VAlign_Bottom)
+		.HAlign(EHorizontalAlignment::HAlign_Right)
+		[
+			SNew(SButton)
+				.Text(FText::FromString(TEXT("Create")))
+				.OnClicked_Lambda([Window, NameValue, ParamValues, InCreationParams, Callback = MoveTemp(InCaptureSourceCreated)]()
+			{
+				FCaptureSourceFrameworkModule& Module = FModuleManager::LoadModuleChecked<FCaptureSourceFrameworkModule>("CaptureSourceFramework");
+				FCaptureSourceManager::FCaptureSourceResult Result = Module.GetCaptureSourceManager().CreateCaptureSource(InCreationParams.GetId(), NameValue->Get<FString>(), *ParamValues);
+
+				if (Result.IsValid())
+				{
+					Callback.ExecuteIfBound(Result.StealValue());
+				}
+
+				Window->RequestDestroyWindow();
+
+				return FReply::Handled();
+			})
+		];
+
+	Window->SetContent(VerticalBox.ToSharedRef());
+
+	return Window;
+}
+
+FString FCaptureSourceDetailsCustomization::GenerateDefaultMapKey(const FString& InName)
+{
+	return GenerateKey(DefaultBase, InName);
+}
+
+bool FCaptureSourceDetailsCustomization::ContainsPropertyWidgetCustomization(const FString& InFactoryId,
+																			 const FString& InCapabilityName,
+																			 const FString& InPropertyName)
+{
+	return PropertyCustomizations.Contains(GenerateKey(InFactoryId, InCapabilityName, InPropertyName));
+}
+
+TSharedPtr<SWidget> FCaptureSourceDetailsCustomization::CreateDefaultPropertyWidget(FCaptureSourceCapability* InCapability,
+																					const FPropertyDesc& InProperty)
+{
+	TSharedPtr<SHorizontalBox> HorizontalBox = SNew(SHorizontalBox);
+
+	FPropertyValue Value = InCapability->GetPropertyValue(InProperty.Name);
+	if (Value.IsType<FEmptyVariantState>())
+	{
+		Value = GetDefaultValueForProperty(InProperty);
+	}
+
+	TSharedPtr<FPropertyValue> PropertyValue = MakeShared<FPropertyValue>(MoveTemp(Value));
+
+	HorizontalBox->AddSlot()
+		.AttachWidget(MakeKeyWidget(InProperty).ToSharedRef());
+
+	HorizontalBox->AddSlot()
+		.AttachWidget(MakeValueWidget(InCapability, InProperty, PropertyValue).ToSharedRef());
+
+	return HorizontalBox;
+}
+
+TSharedPtr<SWidget> FCaptureSourceDetailsCustomization::CreateDefaultCapabilityPropertyWidget(FCaptureSourceCapability* InCapability,
+																							  const FPropertyDesc& InProperty,
+																							  TSharedPtr<FPropertyValue> InPropertyValue)
+{
+	TSharedPtr<SHorizontalBox> HorizontalBox = SNew(SHorizontalBox);
+
+	HorizontalBox->AddSlot()
+		.AttachWidget(MakeKeyWidget(InProperty).ToSharedRef());
+
+	HorizontalBox->AddSlot()
+		.AttachWidget(MakeValueWidget(InCapability, InProperty, MoveTemp(InPropertyValue)).ToSharedRef());
+
+	return HorizontalBox;
+}
+
+TSharedPtr<SWidget> FCaptureSourceDetailsCustomization::CreateDefaultCommandWidget(FCaptureSourceCapability* InCapability,
+																				   const FCommandDesc& InCommand)
+{
+	TSharedPtr<SVerticalBox> VerticalBox = SNew(SVerticalBox);
+	TSharedPtr<FCommandBase::FValues> CommandValues = MakeShared<FCommandBase::FValues>();
+
+	for (const FPropertyDesc& Property : InCommand.Params)
+	{
+		CommandValues->Emplace(Property.Name, GetDefaultValueForProperty(Property));
+
+		TSharedRef<SWidget> Widget = CreateParamWidget(Property, (*CommandValues)[Property.Name]).ToSharedRef();
+
+		VerticalBox->AddSlot()
+			.AutoHeight()
+			.VAlign(EVerticalAlignment::VAlign_Center)
+			.AttachWidget(MoveTemp(Widget));
+	}
+
+	VerticalBox->AddSlot()
+		.AutoHeight()
+		.VAlign(EVerticalAlignment::VAlign_Bottom)
+		.HAlign(EHorizontalAlignment::HAlign_Right)
+		[
+			SNew(SButton)
+				.Text(FText::FromString(TEXT("Execute")))
+				.OnClicked_Lambda([CommandName = InCommand.Name, CommandValues, InCapability]() mutable
+			{
+				TSharedPtr<FCommandBase> CommandPtr = MakeShared<FCommandBase>(CommandName);
+				CommandPtr->SetParamValues(*CommandValues);
+
+				InCapability->ExecuteCommand(MoveTemp(CommandPtr));
+
+				return FReply::Handled();
+			})
+		];
+
+	return VerticalBox;
+}
+
+TSharedPtr<SWidget> FCaptureSourceDetailsCustomization::CreateParamWidget(const FPropertyDesc& InProperty,
+																		  FPropertyValue& InPropertyValue)
+{
+	TSharedPtr<SHorizontalBox> HorizontalBox = SNew(SHorizontalBox);
+
+	HorizontalBox->AddSlot()
+		.AutoWidth()
+		.Padding(5.0f)
+		.AttachWidget(MakeKeyWidget(InProperty).ToSharedRef());
+
+	HorizontalBox->AddSlot()
+		.HAlign(EHorizontalAlignment::HAlign_Fill)
+		.AttachWidget(MakeValueWidget(InProperty, InPropertyValue).ToSharedRef());
+
+	return HorizontalBox;
+}
+
+TSharedPtr<SWidget> FCaptureSourceDetailsCustomization::MakeKeyWidget(const FPropertyDesc& InProperty)
+{
+	TSharedPtr<STextBlock> TextBlock = SNew(STextBlock)
+		.Text(FText::FromString(InProperty.Name))
+		.TextStyle(FAppStyle::Get(), "NormalText");
+
+	return TextBlock;
+}
+
+TSharedPtr<SWidget> FCaptureSourceDetailsCustomization::MakeValueWidget(FCaptureSourceCapability* InCapability, 
+																		const FPropertyDesc& InProperty, 
+																		TSharedPtr<FPropertyValue> InPropertyValue)
+{
+	TSharedPtr<SWidget> Widget;
+
+	switch (InProperty.Type)
+	{
+		case FPropertyDesc::EType::Bool:
+			Widget = SNew(SCheckBox)
+				.OnCheckStateChanged_Lambda([InPropertyValue, InCapability, Name = InProperty.Name](ECheckBoxState InCheckState)
+			{
+				InPropertyValue->Set<bool>(InCheckState == ECheckBoxState::Checked);
+				InCapability->SetPropertyValue(Name, *InPropertyValue);
+				
+			})
+				.IsChecked_Lambda([InPropertyValue]()
+			{
+				return InPropertyValue->Get<bool>() ? ECheckBoxState::Checked : ECheckBoxState::Unchecked;
+			});
+			break;
+		case FPropertyDesc::EType::Number:
+			Widget = SNew(SNumericEntryBox<int64>)
+				.OnValueCommitted_Lambda([InPropertyValue, InCapability, Name = InProperty.Name](int64 InValue, ETextCommit::Type)
+			{
+				InPropertyValue->Set<int64>(InValue);
+				InCapability->SetPropertyValue(Name, *InPropertyValue);
+			})
+				.Value_Lambda([InPropertyValue]()
+			{
+				return InPropertyValue->Get<int64>();
+			});
+			break;
+		case FPropertyDesc::EType::FloatingPoint:
+			Widget = SNew(SNumericEntryBox<double>)
+				.OnValueCommitted_Lambda([InPropertyValue, InCapability, Name = InProperty.Name](double InValue, ETextCommit::Type)
+			{
+				InPropertyValue->Set<double>(InValue);
+				InCapability->SetPropertyValue(Name, *InPropertyValue);
+			})
+				.Value_Lambda([InPropertyValue]()
+			{
+				return InPropertyValue->Get<double>();
+			});
+			break;
+		case FPropertyDesc::EType::String:
+			Widget = SNew(SEditableTextBox)
+				.OnTextCommitted_Lambda([InPropertyValue, InCapability, Name = InProperty.Name](FText InValue, ETextCommit::Type)
+			{
+				InPropertyValue->Set<FString>(InValue.ToString());
+				InCapability->SetPropertyValue(Name, *InPropertyValue);
+			})
+				.Text_Lambda([InPropertyValue]()
+			{
+				return FText::FromString(InPropertyValue->Get<FString>());
+			});
+			break;
+		case FPropertyDesc::EType::ArrayBool:
+			Widget = SNullWidget::NullWidget;
+			break;
+		case FPropertyDesc::EType::ArrayNumber:
+			Widget = SNullWidget::NullWidget;
+			break;
+		case FPropertyDesc::EType::ArrayString:
+			Widget = SNullWidget::NullWidget;
+			break;
+		case FPropertyDesc::EType::ArrayFloatingPoint:
+		default:
+			Widget = SNullWidget::NullWidget;
+	}
+
+	Widget->SetEnabled(InProperty.Access == FPropertyDesc::EAccess::ReadWrite);
+
+	return Widget;
+}
+
+TSharedPtr<SWidget> FCaptureSourceDetailsCustomization::MakeValueWidget(const FPropertyDesc& InProperty, FPropertyValue& InPropertyValue)
+{
+	TSharedPtr<SWidget> Widget;
+	switch (InProperty.Type)
+	{
+		case FPropertyDesc::EType::Bool:
+			Widget = SNew(SCheckBox)
+				.OnCheckStateChanged_Lambda([&InPropertyValue](ECheckBoxState InCheckState)
+			{
+				InPropertyValue.Set<bool>(InCheckState == ECheckBoxState::Checked);
+			})
+				.IsChecked_Lambda([&InPropertyValue]()
+			{
+				return InPropertyValue.Get<bool>() ? ECheckBoxState::Checked : ECheckBoxState::Unchecked;
+			});
+			break;
+		case FPropertyDesc::EType::Number:
+			Widget = SNew(SNumericEntryBox<int64>)
+				.AllowSpin(false)
+				.OnValueCommitted_Lambda([&InPropertyValue](int64 InValue, ETextCommit::Type)
+			{
+				InPropertyValue.Set<int64>(InValue);
+			})
+				.Value_Lambda([&InPropertyValue]()
+			{
+				return InPropertyValue.Get<int64>();
+			});
+			break;
+		case FPropertyDesc::EType::FloatingPoint:
+			Widget = SNew(SNumericEntryBox<double>)
+				.AllowSpin(false)
+				.OnValueCommitted_Lambda([&InPropertyValue](double InValue, ETextCommit::Type)
+			{
+				InPropertyValue.Set<double>(InValue);
+			})
+				.Value_Lambda([&InPropertyValue]()
+			{
+				return InPropertyValue.Get<double>();
+			});
+			break;
+		case FPropertyDesc::EType::String:
+			Widget = SNew(SEditableTextBox)
+				.OnTextCommitted_Lambda([&InPropertyValue](FText InValue, ETextCommit::Type)
+			{
+				InPropertyValue.Set<FString>(InValue.ToString());
+			})
+				.Text_Lambda([&InPropertyValue]()
+			{
+				return FText::FromString(InPropertyValue.Get<FString>());
+			});
+			break;
+		case FPropertyDesc::EType::ArrayBool:
+			Widget = SNullWidget::NullWidget;
+			break;
+		case FPropertyDesc::EType::ArrayNumber:
+			Widget = SNullWidget::NullWidget;
+			break;
+		case FPropertyDesc::EType::ArrayString:
+			Widget = SNullWidget::NullWidget;
+			break;
+		case FPropertyDesc::EType::ArrayFloatingPoint:
+		default:
+			Widget = SNullWidget::NullWidget;
+	}
+
+	Widget->SetEnabled(InProperty.Access == FPropertyDesc::EAccess::ReadWrite);
+
+	return Widget;
+}
+
+FPropertyValue FCaptureSourceDetailsCustomization::GetDefaultValueForProperty(const FPropertyDesc& InProperty)
+{
+	switch (InProperty.Type)
+	{
+		case FPropertyDesc::EType::Bool:
+			return MakePropertyValue(false);
+		case FPropertyDesc::EType::Number:
+			return MakePropertyValue<int64>(0);
+		case FPropertyDesc::EType::FloatingPoint:
+			return MakePropertyValue<double>(0.0);
+		case FPropertyDesc::EType::String:
+			return MakePropertyValue<FString>(TEXT(""));
+		case FPropertyDesc::EType::ArrayBool:
+			return MakePropertyValue(FEmptyVariantState());
+		case FPropertyDesc::EType::ArrayNumber:
+			return MakePropertyValue(FEmptyVariantState());
+		case FPropertyDesc::EType::ArrayString:
+			return MakePropertyValue(FEmptyVariantState());
+		case FPropertyDesc::EType::ArrayFloatingPoint:
+		default:
+			return MakePropertyValue(FEmptyVariantState());
+	}
+}
+
+void FCaptureSourceDetailsCustomization::HandlePropertyChangedEvent(TSharedPtr<const FCaptureEvent> InEvent, 
+																	TMap<FString, TSharedPtr<FPropertyValue>> InPropertyValueHolder)
+{
+	TSharedPtr<const FCapturePropertyChangedEvent> PropertyEvent = StaticCastSharedPtr<const FCapturePropertyChangedEvent>(InEvent);
+
+	if (!InPropertyValueHolder.Contains(PropertyEvent->Name))
+	{
+		return;
+	}
+
+	TSharedPtr<FPropertyValue> BoundValue = InPropertyValueHolder[PropertyEvent->Name];
+	*BoundValue = PropertyEvent->Value;
+}
\ No newline at end of file
diff --git a/Engine/Plugins/VirtualProduction/CaptureManager/Source/CaptureSourceFramework/Private/CaptureSourceError.cpp b/Engine/Plugins/VirtualProduction/CaptureManager/Source/CaptureSourceFramework/Private/CaptureSourceError.cpp
new file mode 100644
--- /dev/null
+++ b/Engine/Plugins/VirtualProduction/CaptureManager/Source/CaptureSourceFramework/Private/CaptureSourceError.cpp
@@ -0,0 +1,19 @@
+// Copyright Epic Games, Inc. All Rights Reserved.
+
+#include "CaptureSourceError.h"
+
+FCaptureSourceError::FCaptureSourceError(ECaptureSourceError InErrorCode, FString InMessage)
+	: ErrorCode(InErrorCode)
+	, Message(MoveTemp(InMessage))
+{
+}
+
+ECaptureSourceError FCaptureSourceError::GetErrorCode() const
+{
+	return ErrorCode;
+}
+
+const FString& FCaptureSourceError::GetMessage() const
+{
+	return Message;
+}
\ No newline at end of file
diff --git a/Engine/Plugins/VirtualProduction/CaptureManager/Source/CaptureSourceFramework/Private/CaptureSourceFactory.cpp b/Engine/Plugins/VirtualProduction/CaptureManager/Source/CaptureSourceFramework/Private/CaptureSourceFactory.cpp
new file mode 100644
--- /dev/null
+++ b/Engine/Plugins/VirtualProduction/CaptureManager/Source/CaptureSourceFramework/Private/CaptureSourceFactory.cpp
@@ -0,0 +1,72 @@
+// Copyright Epic Games, Inc. All Rights Reserved.
+
+#include "CaptureSourceFactory.h"
+
+FCaptureSourceDescriptor::FCaptureSourceDescriptor(const FString& InId, TArray<FPropertyDesc> InCreationParams)
+	: Id(InId)
+	, CreationParams(MoveTemp(InCreationParams))
+{
+}
+
+void FCaptureSourceDescriptor::SetDiscoverable(bool bInIsDiscoverable)
+{
+	bIsDiscoverable = bInIsDiscoverable;
+}
+
+bool FCaptureSourceDescriptor::IsDiscoverable() const
+{
+	return bIsDiscoverable;
+}
+
+void FCaptureSourceDescriptor::AddCreationParam(FPropertyDesc InParam)
+{
+	CreationParams.Push(MoveTemp(InParam));
+}
+
+const TArray<FPropertyDesc>& FCaptureSourceDescriptor::GetCreationParams() const
+{
+	return CreationParams;
+}
+
+const FString& FCaptureSourceDescriptor::GetId() const
+{
+	return Id;
+}
+
+FCaptureVoidResult FCaptureSourceFactory::Discover(FDiscoverCallback InCallback)
+{
+	if (!IsDiscoverable())
+	{
+		InCallback.ExecuteIfBound(TArray<FCaptureSourceResult>());
+
+		return MakeError(FCaptureSourceError(ECaptureSourceError::NotSupported));
+	}
+
+	FCaptureVoidResult Result = MakeValue();
+
+	DiscoverImpl(MoveTemp(InCallback), Result);
+
+	return Result;
+}
+
+FCaptureSourceFactory::FCaptureSourceResult FCaptureSourceFactory::CreateCaptureSource(const FString& InName, TMap<FString, FPropertyValue> InCreationParamsValue)
+{
+	FCaptureSourceResult CaptureSourceResult = CreateCaptureSourceImpl(InName, InCreationParamsValue);
+
+	if (CaptureSourceResult.IsValid())
+	{
+		CaptureSourceResult.GetValue()->SetCreationParams(MoveTemp(InCreationParamsValue));
+	}
+
+	return CaptureSourceResult;
+}
+
+FCaptureSourceDescriptor FCaptureSourceFactory::CreateCaptureSourceDescriptor() const
+{
+	TArray<FPropertyDesc> MandatoryCreationParams;
+	FCaptureSourceDescriptor Descriptor(GetCaptureSourceFactoryId(), MoveTemp(MandatoryCreationParams));
+
+	CreateCaptureSourceDescriptorImpl(Descriptor);
+
+	return Descriptor;
+}
\ No newline at end of file
diff --git a/Engine/Plugins/VirtualProduction/CaptureManager/Source/CaptureSourceFramework/Private/CaptureSourceFrameworkLog.h b/Engine/Plugins/VirtualProduction/CaptureManager/Source/CaptureSourceFramework/Private/CaptureSourceFrameworkLog.h
new file mode 100644
--- /dev/null
+++ b/Engine/Plugins/VirtualProduction/CaptureManager/Source/CaptureSourceFramework/Private/CaptureSourceFrameworkLog.h
@@ -0,0 +1,7 @@
+// Copyright Epic Games, Inc. All Rights Reserved.
+
+#pragma once
+
+#include "Logging/LogMacros.h"
+
+DECLARE_LOG_CATEGORY_EXTERN(LogCaptureSourceManager, Log, All);
diff --git a/Engine/Plugins/VirtualProduction/CaptureManager/Source/CaptureSourceFramework/Private/CaptureSourceFrameworkModule.cpp b/Engine/Plugins/VirtualProduction/CaptureManager/Source/CaptureSourceFramework/Private/CaptureSourceFrameworkModule.cpp
new file mode 100644
--- /dev/null
+++ b/Engine/Plugins/VirtualProduction/CaptureManager/Source/CaptureSourceFramework/Private/CaptureSourceFrameworkModule.cpp
@@ -0,0 +1,33 @@
+// Copyright Epic Games, Inc. All Rights Reserved.
+
+#include "CaptureSourceFrameworkModule.h"
+
+void FCaptureSourceFrameworkModule::StartupModule()
+{
+	CaptureManager = MakeUnique<FCaptureSourceManager>();
+	Customization = MakeUnique<FCaptureSourceDetailsCustomization>();
+}
+
+void FCaptureSourceFrameworkModule::ShutdownModule()
+{
+	for (FCaptureSourceId Id : CaptureManager->GetCaptureSourceIds())
+	{
+		CaptureManager->RemoveCaptureSource(Id);
+	}
+
+	CaptureManager = nullptr;
+
+	Customization = nullptr;
+}
+
+FCaptureSourceManager& FCaptureSourceFrameworkModule::GetCaptureSourceManager()
+{
+	return *CaptureManager;
+}
+
+FCaptureSourceDetailsCustomization& FCaptureSourceFrameworkModule::GetCustomization()
+{
+	return *Customization;
+}
+
+IMPLEMENT_MODULE(FCaptureSourceFrameworkModule, CaptureSourceFramework)
\ No newline at end of file
diff --git a/Engine/Plugins/VirtualProduction/CaptureManager/Source/CaptureSourceFramework/Private/CaptureSourceHandle.cpp b/Engine/Plugins/VirtualProduction/CaptureManager/Source/CaptureSourceFramework/Private/CaptureSourceHandle.cpp
new file mode 100644
--- /dev/null
+++ b/Engine/Plugins/VirtualProduction/CaptureManager/Source/CaptureSourceFramework/Private/CaptureSourceHandle.cpp
@@ -0,0 +1,58 @@
+// Copyright Epic Games, Inc. All Rights Reserved.
+
+#include "CaptureSourceHandle.h"
+
+FCaptureSourceHandle::FCaptureSourceHandle(FCaptureSource* InCaptureSource)
+	: CaptureSource(InCaptureSource)
+{
+}
+
+FCaptureVoidResult FCaptureSourceHandle::Start()
+{
+	return CaptureSource->Start();
+}
+
+FCaptureVoidResult FCaptureSourceHandle::Stop()
+{
+	return CaptureSource->Stop();
+}
+
+FCaptureSourceCapability* FCaptureSourceHandle::GetCapability(FString InName) const
+{
+	return CaptureSource->GetCapability(MoveTemp(InName));
+}
+
+TArray<FString> FCaptureSourceHandle::GetCapabilities() const
+{
+	return CaptureSource->GetCapabilities();
+}
+
+const FString& FCaptureSourceHandle::GetCaptureSourceFactoryId() const
+{
+	return CaptureSource->GetCaptureSourceFactoryId();
+}
+
+const FString& FCaptureSourceHandle::GetName() const
+{
+	return CaptureSource->GetName();
+}
+
+TArray<FString> FCaptureSourceHandle::GetAvailableEvents() const
+{
+	return CaptureSource->GetAvailableEvents();
+}
+
+void FCaptureSourceHandle::SubscribeToEvent(const FString& InEventName, FCaptureEventHandler InHandler)
+{
+	CaptureSource->SubscribeToEvent(InEventName, MoveTemp(InHandler));
+}
+
+void FCaptureSourceHandle::UnsubscribeAll()
+{
+	CaptureSource->UnsubscribeAll();
+}
+
+const TMap<FString, FPropertyValue>& FCaptureSourceHandle::GetCreationParams() const
+{
+	return CaptureSource->GetCreationParams();
+}
\ No newline at end of file
diff --git a/Engine/Plugins/VirtualProduction/CaptureManager/Source/CaptureSourceFramework/Private/CaptureSourceManager.cpp b/Engine/Plugins/VirtualProduction/CaptureManager/Source/CaptureSourceFramework/Private/CaptureSourceManager.cpp
new file mode 100644
--- /dev/null
+++ b/Engine/Plugins/VirtualProduction/CaptureManager/Source/CaptureSourceFramework/Private/CaptureSourceManager.cpp
@@ -0,0 +1,237 @@
+// Copyright Epic Games, Inc. All Rights Reserved.
+
+#include "CaptureSourceManager.h"
+
+#include "CaptureSourceFrameworkLog.h"
+
+DEFINE_LOG_CATEGORY(LogCaptureSourceManager);
+
+template<typename T>
+class FMonitor
+{
+public:
+	template<typename ...Args>
+	FMonitor(Args&&... InArgs) 
+		: Object(Forward<Args>(InArgs)...)
+	{
+
+	}
+
+	FMonitor(T InObject) 
+		: Object(MoveTemp(InObject))
+	{
+	}
+
+	class FHelper
+	{
+	public:
+		FHelper(FMonitor* InOwner)
+			: Owner(InOwner)
+			, ScopeLock(&InOwner->Mutex)
+		{
+		}
+
+		T* operator->()
+		{
+			return &Owner->Object;
+		}
+
+	private:
+
+		FMonitor* Owner;
+		FScopeLock ScopeLock;
+	};
+
+	FHelper operator->()
+	{
+		return FHelper(this);
+	}
+
+	FHelper Lock()
+	{
+		return FHelper(this);
+	}
+
+	T& GetUnsafe()
+	{
+		return Object;
+	}
+
+	T Claim()
+	{
+		return MoveTemp(Object);
+	}
+
+private:
+
+	FCriticalSection Mutex;
+	T Object;
+};
+
+FCaptureSourceManager::FCaptureSourceManager()
+{
+}
+
+FCaptureSourceHandle FCaptureSourceManager::GetCaptureSource(FCaptureSourceId InCaptureSourceId)
+{
+	check(CaptureSources.Contains(InCaptureSourceId));
+	return FCaptureSourceHandle(CaptureSources[InCaptureSourceId].Get());
+}
+
+TArray<FCaptureSourceId> FCaptureSourceManager::GetCaptureSourceIds() const
+{
+	TArray<FCaptureSourceId> CaptureSourceIds;
+	CaptureSourceIds.Reserve(CaptureSources.Num());
+
+	CaptureSources.GetKeys(CaptureSourceIds);
+
+	return CaptureSourceIds;
+}
+
+void FCaptureSourceManager::RemoveCaptureSource(FCaptureSourceId InCaptureSourceId)
+{
+	check(CaptureSources.Contains(InCaptureSourceId));
+	CaptureSources.Remove(InCaptureSourceId);
+}
+
+void FCaptureSourceManager::DiscoverCaptureSources(FDiscoverCallback InCallback)
+{
+	TSharedPtr<FCallbackSynchronizer> CallbackSynchronizer = FCallbackSynchronizer::Create();
+
+	TSharedPtr<FMonitor<TArray<FCaptureSourceResult>>> Collector
+		= MakeShared<FMonitor<TArray<FCaptureSourceResult>>>();
+
+	for (const TPair<FString, TUniquePtr<FCaptureSourceFactory>>& CaptureSourceFactoryPair : CaptureSourceFactoryList)
+	{
+		const TUniquePtr<FCaptureSourceFactory>& CaptureSourceFactory = CaptureSourceFactoryPair.Value;
+
+		auto Callback = CallbackSynchronizer->CreateCallback(
+			[this, CaptureSourceFactoryId = CaptureSourceFactory->GetCaptureSourceFactoryId(), Collector = Collector.Get()](TArray<FCaptureSourceFactory::FCaptureSourceResult> InCaptureSourceResults)
+		{
+			for (FCaptureSourceFactory::FCaptureSourceResult& CaptureSourceResult : InCaptureSourceResults)
+			{
+				if (CaptureSourceResult.IsValid())
+				{
+					FCaptureSourceId Id = CurrentCaptureSourceId++;
+					CaptureSources.Emplace(Id, CaptureSourceResult.StealValue());
+
+					(*Collector)->Emplace(MakeValue(Id));
+				}
+				else
+				{
+					FCaptureSourceError Error = CaptureSourceResult.StealError();
+					UE_LOG(LogCaptureSourceManager, Error, TEXT("Failed to create a capture source %s: %s"), *CaptureSourceFactoryId, *Error.GetMessage());
+
+					(*Collector)->Emplace(MakeError(MoveTemp(Error)));
+				}
+			}
+		});
+
+		FCaptureVoidResult Result = CaptureSourceFactory->Discover(FCaptureSourceFactory::FDiscoverCallback::CreateLambda(MoveTemp(Callback)));
+
+		if (Result.HasError())
+		{
+			UE_LOG(LogCaptureSourceManager, Error, TEXT("Failed to start discovery process for %s: %s"), *CaptureSourceFactory->GetCaptureSourceFactoryId(), *Result.GetError().GetMessage());
+		}
+	}
+
+	CallbackSynchronizer->AfterAll(
+		FCallbackSynchronizer::FAfterAllDelegate::CreateLambda([this, UserCallback = MoveTemp(InCallback), Collector = MoveTemp(Collector)]()
+	{
+		UserCallback.ExecuteIfBound(Collector->Claim());
+	}), true);
+}
+
+void FCaptureSourceManager::DiscoverCaptureSourcesForFactory(FString InCaptureSourceFactoryId, FDiscoverCallback InCallback)
+{
+	check(CaptureSourceFactoryList.Contains(InCaptureSourceFactoryId));
+	
+	const TUniquePtr<FCaptureSourceFactory>& CaptureSourceFactory = CaptureSourceFactoryList[InCaptureSourceFactoryId];
+
+	FCaptureSourceFactory::FDiscoverCallback Callback = 
+		FCaptureSourceFactory::FDiscoverCallback::CreateLambda(
+			[this, InCaptureSourceFactoryId, UserCallback = MoveTemp(InCallback)](TArray<FCaptureSourceFactory::FCaptureSourceResult> InCaptureSourceResults)
+	{
+		TArray<FCaptureSourceResult> CaptureSourceIds;
+		for (FCaptureSourceFactory::FCaptureSourceResult& CaptureSourceResult : InCaptureSourceResults)
+		{
+			if (CaptureSourceResult.IsValid())
+			{
+				FCaptureSourceId Id = CurrentCaptureSourceId++;
+				CaptureSources.Emplace(Id, CaptureSourceResult.StealValue());
+
+				CaptureSourceIds.Emplace(MakeValue(Id));
+			}
+			else
+			{
+				FCaptureSourceError Error = CaptureSourceResult.StealError();
+				UE_LOG(LogCaptureSourceManager, Error, TEXT("Failed to create a capture source %s: %s"), *InCaptureSourceFactoryId, *Error.GetMessage());
+
+				CaptureSourceIds.Emplace(MakeError(MoveTemp(Error)));
+			}
+		}
+
+		UserCallback.ExecuteIfBound(MoveTemp(CaptureSourceIds));
+	});
+
+	FCaptureVoidResult Result = CaptureSourceFactory->Discover(MoveTemp(Callback));
+
+	if (Result.HasError())
+	{
+		UE_LOG(LogCaptureSourceManager, Error, TEXT("Failed to start discovery process for %s: %s"), *InCaptureSourceFactoryId, *Result.GetError().GetMessage());
+	}
+}
+
+void FCaptureSourceManager::RegisterCaptureSourceFactory(TUniquePtr<FCaptureSourceFactory> InCaptureSourceFactory)
+{
+	FString CaptureSourceFactoryId = InCaptureSourceFactory->GetCaptureSourceFactoryId();
+	CaptureSourceFactoryList.Emplace(MoveTemp(CaptureSourceFactoryId), MoveTemp(InCaptureSourceFactory));
+}
+
+FCaptureSourceDescriptor FCaptureSourceManager::GetCaptureSourceFactory(FString InCaptureSourceFactoryId)
+{
+	return CaptureSourceFactoryList[InCaptureSourceFactoryId]->CreateCaptureSourceDescriptor();
+}
+
+TArray<FCaptureSourceDescriptor> FCaptureSourceManager::GetCaptureSourceFactoryList()
+{
+	TArray<FCaptureSourceDescriptor> Descriptors;
+	Descriptors.Reserve(CaptureSourceFactoryList.Num());
+
+	for (const TPair<FString, TUniquePtr<FCaptureSourceFactory>>& CaptureSourceFactory : CaptureSourceFactoryList)
+	{
+		Descriptors.Add(CaptureSourceFactory.Value->CreateCaptureSourceDescriptor());
+	}
+
+	return Descriptors;
+}
+
+FCaptureSourceManager::FCaptureSourceResult FCaptureSourceManager::CreateCaptureSource(FString InCaptureSourceFactoryId, FString InName, TMap<FString, FPropertyValue> InCreationParamsValue)
+{
+	TUniquePtr<FCaptureSourceFactory>* CaptureSourceFactory = CaptureSourceFactoryList.Find(InCaptureSourceFactoryId);
+	
+	if (!CaptureSourceFactory)
+	{
+		return MakeError(FCaptureSourceError(ECaptureSourceError::NotFound, TEXT("Capture Source Factory doesn't exist")));
+	}
+
+	FCaptureSourceFactory::FCaptureSourceResult CaptureSourceResult = (*CaptureSourceFactory)->CreateCaptureSource(InName, MoveTemp(InCreationParamsValue));
+
+	if (CaptureSourceResult.IsValid())
+	{
+		FCaptureSourceId Id = CurrentCaptureSourceId++;
+		CaptureSources.Emplace(Id, CaptureSourceResult.StealValue());
+
+		return MakeValue(Id);
+	}
+	else
+	{
+		return MakeError(CaptureSourceResult.StealError());
+	}
+}
+
+FString FCaptureSourceManager::GetCaptureSourceFactoryById(FCaptureSourceId InCaptureSourceId)
+{
+	check(CaptureSources.Contains(InCaptureSourceId));
+	return CaptureSources[InCaptureSourceId]->GetCaptureSourceFactoryId();
+}
\ No newline at end of file
diff --git a/Engine/Plugins/VirtualProduction/CaptureManager/Source/CaptureSourceFramework/Private/Tests/TestCaptureSource.cpp b/Engine/Plugins/VirtualProduction/CaptureManager/Source/CaptureSourceFramework/Private/Tests/TestCaptureSource.cpp
new file mode 100644
--- /dev/null
+++ b/Engine/Plugins/VirtualProduction/CaptureManager/Source/CaptureSourceFramework/Private/Tests/TestCaptureSource.cpp
@@ -0,0 +1,125 @@
+// Copyright Epic Games, Inc. All Rights Reserved.
+
+#include "CaptureSource.h"
+
+#include "Misc/AutomationTest.h"
+
+#if WITH_DEV_AUTOMATION_TESTS
+
+BEGIN_DEFINE_SPEC(FCaptureSourceTest, "Plugin.CaptureSourceFramework.CaptureSource", EAutomationTestFlags::EditorContext | EAutomationTestFlags::ProductFilter | EAutomationTestFlags::MediumPriority)
+
+static const FString TestCaptureSourceFactoryId;
+static const FString TestCaptureSourceName;
+
+class FTestCapability final : public FCaptureSourceCapability
+{
+public:
+	static const FString Name;
+
+	static const FString TestEvent;
+
+	FTestCapability()
+		: FCaptureSourceCapability(Name)
+	{
+		RegisterEvent(TestEvent);
+	}
+};
+
+class FTestCaptureSource final : public FCaptureSource
+{
+public:
+	FTestCaptureSource(const FString& InName) :
+		FCaptureSource(TestCaptureSourceFactoryId, InName)
+	{
+		TUniquePtr<FTestCapability> TestCapability = MakeUnique<FTestCapability>();
+
+		AddCapability(TestCapability.Get());
+
+		OwnedCapabilities.Add(MoveTemp(TestCapability));
+	}
+
+	virtual FCaptureVoidResult Start() final
+	{
+		return MakeValue();
+	}
+
+	virtual FCaptureVoidResult Stop() final
+	{
+		return MakeValue();
+	}
+
+private:
+
+	TArray<TUniquePtr<FCaptureSourceCapability>> OwnedCapabilities;
+};
+
+TUniquePtr<FCaptureSource> CaptureSource;
+
+END_DEFINE_SPEC(FCaptureSourceTest)
+
+const FString FCaptureSourceTest::TestCaptureSourceFactoryId = TEXT("TestCaptureSourceFactoryId");
+const FString FCaptureSourceTest::TestCaptureSourceName = TEXT("TestCaptureSource");
+
+const FString FCaptureSourceTest::FTestCapability::Name = TEXT("TestCapability");
+const FString FCaptureSourceTest::FTestCapability::TestEvent = TEXT("TestEvent");
+
+void FCaptureSourceTest::Define()
+{
+	CaptureSource = MakeUnique<FTestCaptureSource>(TestCaptureSourceName);
+
+	Describe("GetCapabilities()", [this]()
+	{
+		It("should return all registered capabilities", [this]()
+		{
+			const TArray<FString> Capabilities = CaptureSource->GetCapabilities();
+
+			UTEST_EQUAL("One capability is registered", Capabilities.Num(), 1);
+			UTEST_EQUAL("Registered capability must match", Capabilities[0], FTestCapability::Name);
+			return true;
+		});
+	});
+
+	Describe("GetCapability()", [this]()
+	{
+		It("should return registered capability", [this]()
+		{
+			FCaptureSourceCapability* TestCapability = CaptureSource->GetCapability(FTestCapability::Name);
+
+			UTEST_NOT_NULL("Fetched capability is valid", TestCapability);
+			UTEST_EQUAL("Fetched capability name must match", TestCapability->GetName(), FTestCapability::Name);
+			return true;
+		});
+
+		It("should return invalid capability (nullptr)", [this]()
+		{
+			FCaptureSourceCapability* NotExistingCapability = CaptureSource->GetCapability(TEXT("NotExistingCapability"));
+
+			UTEST_NULL("Fetched capability is nullptr", NotExistingCapability);
+			return true;
+		});
+	});
+
+	Describe("GetCaptureSourceFactoryId()", [this]()
+	{
+		It("should return correct capture source factory id", [this]()
+		{
+			const FString CaptureSourceFactoryId = CaptureSource->GetCaptureSourceFactoryId();
+
+			UTEST_EQUAL("Capture source factory id must match", CaptureSourceFactoryId, TestCaptureSourceFactoryId);
+			return true;
+		});
+	});
+
+	Describe("GetName()", [this]()
+	{
+		It("should return correct capture source name", [this]()
+		{
+			const FString Name = CaptureSource->GetName();
+
+			UTEST_EQUAL("Capture source name must match", Name, TestCaptureSourceName);
+			return true;
+		});
+	});
+}
+
+#endif // WITH_DEV_AUTOMATION_TESTS
\ No newline at end of file
diff --git a/Engine/Plugins/VirtualProduction/CaptureManager/Source/CaptureSourceFramework/Private/Tests/TestCaptureSourceCapability.cpp b/Engine/Plugins/VirtualProduction/CaptureManager/Source/CaptureSourceFramework/Private/Tests/TestCaptureSourceCapability.cpp
new file mode 100644
--- /dev/null
+++ b/Engine/Plugins/VirtualProduction/CaptureManager/Source/CaptureSourceFramework/Private/Tests/TestCaptureSourceCapability.cpp
@@ -0,0 +1,275 @@
+// Copyright Epic Games, Inc. All Rights Reserved.
+
+#include "CaptureSourceCapability.h"
+
+#include "Misc/AutomationTest.h"
+
+#if WITH_DEV_AUTOMATION_TESTS
+
+BEGIN_DEFINE_SPEC(FCaptureSourceCapabilityTest, "Plugin.CaptureSourceFramework.CaptureSourceCapability", EAutomationTestFlags::EditorContext | EAutomationTestFlags::ProductFilter | EAutomationTestFlags::MediumPriority)
+
+class FTestCommand final : public FCommandBase
+{
+public:
+	static const FString Command;
+	static const FString CommandParam;
+
+	FTestCommand()
+		: FCommandBase(Command)
+		, bReturnValue(false)
+	{
+	}
+
+	bool bReturnValue;
+};
+
+class FTestCapability final : public FCaptureSourceCapability
+{
+public:
+	static const FString Name;
+
+	static const FString ReadWriteProperty;
+	static const FString ReadOnlyProperty;
+
+	static const FString TestEvent;
+
+	static const FString DefaultROPropertyValue;
+
+	FTestCapability()
+		: FCaptureSourceCapability(Name)
+	{
+		AddProperty(
+			FPropertyDesc(ReadWriteProperty, 
+						  FPropertyDesc::EType::String));
+		AddProperty(
+			FPropertyDesc(ReadOnlyProperty, 
+						  FPropertyDesc::EType::String, 
+						  FPropertyDesc::EAccess::ReadOnly));
+		AddCommand(
+			FCommandDesc(FTestCommand::Command,
+						 {
+							 FPropertyDesc(FTestCommand::CommandParam, FPropertyDesc::EType::String)
+						 }));
+
+		RegisterEvent(TestEvent);
+
+		SetPropertyGetter(ReadWriteProperty, FPropertyGetter::CreateRaw(this, &FTestCapability::GetPropertyValues));
+		SetPropertySetter(ReadWriteProperty, FPropertySetter::CreateRaw(this, &FTestCapability::SetRWPropertyValue));
+
+		SetPropertyGetter(ReadOnlyProperty, FPropertyGetter::CreateRaw(this, &FTestCapability::GetPropertyValues));
+
+		SetCommandHandler(FTestCommand::Command, FCommandHandler::CreateRaw(this, &FTestCapability::HandleCommand));
+
+		ROPropertyValue = DefaultROPropertyValue;
+	}
+
+private:
+
+	FPropertyValue GetPropertyValues(const FString& InName) const
+	{
+		if (ReadWriteProperty == InName)
+		{
+			return MakePropertyValue(RWPropertyValue);
+		}
+		else
+		{
+			return MakePropertyValue(ROPropertyValue);
+		}
+	}
+
+	void SetRWPropertyValue(const FString& InName, FPropertyValue InValue)
+	{
+		RWPropertyValue = InValue.Get<FString>();
+
+		PublishPropertyChangedEvent(InName, MoveTemp(InValue));
+	}
+
+	void HandleCommand(TSharedPtr<FCommandBase> InCommand)
+	{
+		StaticCastSharedPtr<FTestCommand>(InCommand)->bReturnValue = true;
+	}
+
+	FString RWPropertyValue;
+	FString ROPropertyValue;
+};
+
+TUniquePtr<FCaptureSourceCapability> TestCapability;
+
+END_DEFINE_SPEC(FCaptureSourceCapabilityTest)
+
+bool operator==(const FPropertyDesc& InLeft, const FPropertyDesc& InRight)
+{
+	return InLeft.Name == InRight.Name &&
+		InLeft.Type == InRight.Type &&
+		InLeft.Access == InRight.Access;
+}
+
+bool operator==(const FCommandDesc& InLeft, const FCommandDesc& InRight)
+{
+	return InLeft.Name == InRight.Name && 
+		InLeft.Params == InRight.Params;
+}
+
+const FString FCaptureSourceCapabilityTest::FTestCommand::Command = TEXT("Command");
+const FString FCaptureSourceCapabilityTest::FTestCommand::CommandParam = TEXT("CommandParam");
+
+const FString FCaptureSourceCapabilityTest::FTestCapability::Name = TEXT("TestCapability");
+const FString FCaptureSourceCapabilityTest::FTestCapability::ReadWriteProperty = TEXT("ReadWriteProperty");
+const FString FCaptureSourceCapabilityTest::FTestCapability::ReadOnlyProperty = TEXT("ReadOnlyProperty");
+const FString FCaptureSourceCapabilityTest::FTestCapability::TestEvent = TEXT("TestEvent");
+
+const FString FCaptureSourceCapabilityTest::FTestCapability::DefaultROPropertyValue = TEXT("Default");
+
+void FCaptureSourceCapabilityTest::Define()
+{
+	TestCapability = MakeUnique<FTestCapability>();
+
+	Describe("GetProperties()", [this]()
+	{
+		It("should return all registered properties", [this]()
+		{
+			const TArray<FPropertyDesc> Properties = TestCapability->GetProperties();
+
+			UTEST_EQUAL("Two properties are registered", Properties.Num(), 2);
+
+			UTEST_EQUAL("First property name must match", Properties[0].Name, FTestCapability::ReadWriteProperty);
+			UTEST_EQUAL("First property type must match", Properties[0].Type, FPropertyDesc::EType::String);
+			UTEST_EQUAL("First property access must match", Properties[0].Access, FPropertyDesc::EAccess::ReadWrite);
+
+			UTEST_EQUAL("Second property name must match", Properties[1].Name, FTestCapability::ReadOnlyProperty);
+			UTEST_EQUAL("Second property type must match", Properties[1].Type, FPropertyDesc::EType::String);
+			UTEST_EQUAL("Second property access must match", Properties[1].Access, FPropertyDesc::EAccess::ReadOnly);
+
+			return true;
+		});
+	});
+
+	Describe("GetCommands()", [this]()
+	{
+		It("should return all registered commands", [this]()
+		{
+			const TArray<FCommandDesc> Commands = TestCapability->GetCommands();
+
+			UTEST_EQUAL("One command is registered", Commands.Num(), 1);
+			UTEST_EQUAL("First command name must match", Commands[0].Name, FTestCommand::Command);
+
+			const TArray<FPropertyDesc>& Params = Commands[0].Params;
+			UTEST_EQUAL("First param name must match", Params[0].Name, FTestCommand::CommandParam);
+			UTEST_EQUAL("First param type must match", Params[0].Type, FPropertyDesc::EType::String);
+			UTEST_EQUAL("First param access must match", Params[0].Access, FPropertyDesc::EAccess::ReadWrite);
+
+			return true;
+		});
+	});
+
+	Describe("PublishPropertyChangedEvent()", [this]()
+	{
+		FString Value = TEXT("Value");
+		TSharedPtr<bool> EventReturn = MakeShared<bool>(false);
+
+		FCaptureEventHandler Handler = FCaptureEventHandler(
+			[this, Value, EventReturn](TSharedPtr<const FCaptureEvent> InEvent)
+		{
+			TSharedPtr<const FCapturePropertyChangedEvent> Event = StaticCastSharedPtr<const FCapturePropertyChangedEvent>(InEvent);
+
+			if (!TestEqual("Event should carry property name that is set", Event->Name, FTestCapability::ReadOnlyProperty))
+			{
+				*EventReturn = false;
+			}
+
+			if (!TestEqual("Event should carry value that is set", Event->Value.Get<FString>(), Value))
+			{
+				*EventReturn = false;
+			}
+
+			*EventReturn = true;
+
+		}, EDelegateExecutionThread::InternalThread);
+
+		TestCapability->SubscribeToEvent(FCapturePropertyChangedEvent::EventName, MoveTemp(Handler));
+
+		It("should correctly publish the event", [this, Value, EventReturn]()
+		{
+			TestCapability->PublishPropertyChangedEvent(FTestCapability::ReadOnlyProperty, MakePropertyValue(Value));
+
+			UTEST_TRUE("Event values are correct", *EventReturn);
+
+			TestCapability->UnsubscribeAll();
+
+			return true;
+		});
+	});
+
+	Describe("GetPropertyValue()", [this]()
+	{
+		It("should get the correct value of the property", [this]()
+		{
+			FPropertyValue PropertyValue = TestCapability->GetPropertyValue(FTestCapability::ReadOnlyProperty);
+
+			UTEST_EQUAL("Values should match", FTestCapability::DefaultROPropertyValue, PropertyValue.Get<FString>());
+
+			return true;
+		});
+	});
+
+	Describe("SetPropertyValue()", [this]()
+	{
+		It("should set the correct value to the property", [this]()
+		{
+			FString Value = TEXT("Value");
+
+			TestCapability->SetPropertyValue(FTestCapability::ReadWriteProperty, MakePropertyValue(Value));
+			
+			FPropertyValue PropertyValue = TestCapability->GetPropertyValue(FTestCapability::ReadWriteProperty);
+
+			UTEST_EQUAL("Set value should match fetched value", Value, PropertyValue.Get<FString>());
+
+			return true;
+		});
+	});
+
+	Describe("ExecuteCommand()", [this]()
+	{
+		It("should correctly execute a command", [this]()
+		{
+			TSharedPtr<FTestCommand> Command = MakeShared<FTestCommand>();
+			Command->AddParamValue(FTestCommand::CommandParam, MakePropertyValue<FString>(TEXT("TestValue")));
+
+			TestCapability->ExecuteCommand(Command);
+
+			UTEST_TRUE("Return value must be true", Command->bReturnValue);
+
+			return true;
+		});
+	});
+
+	Describe("GetName()", [this]()
+	{
+		It("should get the correct capability name", [this]()
+		{
+			FString Name = TestCapability->GetName();
+
+			UTEST_EQUAL("Name values must match", Name, FTestCapability::Name);
+
+			return true;
+		});
+	});
+
+	Describe("CreateDescriptor()", [this]()
+	{
+		It("should get the valid capability descriptor", [this]()
+		{
+			FCaptureSourceCapabilityDesc Descriptor = TestCapability->CreateDescriptor();
+
+			TArray<FPropertyDesc> Properties = TestCapability->GetProperties();
+
+			UTEST_EQUAL("Name values must match", Descriptor.Name, TestCapability->GetName());
+			UTEST_EQUAL("Properties must match", Descriptor.Properties, TestCapability->GetProperties());
+			UTEST_EQUAL("Commands must match", Descriptor.Commands, TestCapability->GetCommands());
+
+			return true;
+		});
+	});
+}
+
+#endif // WITH_DEV_AUTOMATION_TESTS
\ No newline at end of file
diff --git a/Engine/Plugins/VirtualProduction/CaptureManager/Source/CaptureSourceFramework/Private/Tests/TestCaptureSourceFactory.cpp b/Engine/Plugins/VirtualProduction/CaptureManager/Source/CaptureSourceFramework/Private/Tests/TestCaptureSourceFactory.cpp
new file mode 100644
--- /dev/null
+++ b/Engine/Plugins/VirtualProduction/CaptureManager/Source/CaptureSourceFramework/Private/Tests/TestCaptureSourceFactory.cpp
@@ -0,0 +1,204 @@
+// Copyright Epic Games, Inc. All Rights Reserved.
+
+#include "CaptureSourceFactory.h"
+
+#include "Misc/AutomationTest.h"
+
+#if WITH_DEV_AUTOMATION_TESTS
+
+BEGIN_DEFINE_SPEC(FCaptureSourceFactoryTest, "Plugin.CaptureSourceFramework.CaptureSourceFactory", EAutomationTestFlags::EditorContext | EAutomationTestFlags::ProductFilter | EAutomationTestFlags::MediumPriority)
+
+class FTestCaptureSource final : public FCaptureSource
+{
+public:
+	FTestCaptureSource(const FString& InName, const FString& InParamValue) 
+		: FCaptureSource(FTestCaptureSourceFactory::Id, InName)
+		, ParamValue(InParamValue)
+	{
+	}
+
+	virtual FCaptureVoidResult Start() final
+	{
+		return MakeValue();
+	}
+
+	virtual FCaptureVoidResult Stop() final
+	{
+		return MakeValue();
+	}
+
+	FString ParamValue;
+};
+
+class FTestCaptureSourceFactory final : public FCaptureSourceFactory
+{
+public:
+
+	static const FString Id;
+
+	static const FString CreationParam;
+
+	virtual FString GetCaptureSourceFactoryName() const override
+	{
+		return TEXT("Test Capture Source Factory");
+	}
+
+	virtual FString GetCaptureSourceFactoryId() const override
+	{
+		return Id;
+	}
+
+	virtual bool IsDiscoverable() const override
+	{
+		return bIsDiscoverable;
+	}
+
+	void SetDiscoverable(bool bInIsDiscoverable)
+	{
+		bIsDiscoverable = bInIsDiscoverable;
+	}
+
+private:
+
+	virtual void DiscoverImpl(FDiscoverCallback InCallback, FCaptureVoidResult& OutResult) override
+	{
+		TArray<FCaptureSourceResult> DiscoveredCaptureSources;
+
+		DiscoveredCaptureSources.Add(MakeValue(MakeUnique<FTestCaptureSource>(TEXT("DefaultName"), TEXT("DefaultCreationParam"))));
+
+		InCallback.ExecuteIfBound(MoveTemp(DiscoveredCaptureSources));
+
+		OutResult = MakeValue();
+	}
+
+	virtual FCaptureSourceResult CreateCaptureSourceImpl(const FString& InName, const TMap<FString, FPropertyValue>& InCreationParamsValue) override
+	{
+		if (InCreationParamsValue.Num() != 1 && !InCreationParamsValue.Contains(CreationParam))
+		{
+			return MakeError(FCaptureSourceError(ECaptureSourceError::InvalidArgument));
+		}
+
+		return MakeValue(MakeUnique<FTestCaptureSource>(InName, InCreationParamsValue[CreationParam].Get<FString>()));
+	}
+
+	virtual void CreateCaptureSourceDescriptorImpl(FCaptureSourceDescriptor& OutDescriptor) const override
+	{
+		OutDescriptor.SetDiscoverable(bIsDiscoverable);
+
+		OutDescriptor.EmplaceCreationParam(CreationParam, FPropertyDesc::EType::String);
+	}
+
+	bool bIsDiscoverable;
+};
+
+TUniquePtr<FTestCaptureSourceFactory> CaptureSourceFactory;
+
+END_DEFINE_SPEC(FCaptureSourceFactoryTest)
+
+const FString FCaptureSourceFactoryTest::FTestCaptureSourceFactory::Id = TEXT("TestCaptureSourceFactoryId");
+const FString FCaptureSourceFactoryTest::FTestCaptureSourceFactory::CreationParam = TEXT("TestCreationParam");
+
+void FCaptureSourceFactoryTest::Define()
+{
+	CaptureSourceFactory = MakeUnique<FTestCaptureSourceFactory>();
+
+	Describe("Discover()", [this]()
+	{
+		It("runs discovery without problems as factory is discoverable", [this]()
+		{
+			CaptureSourceFactory->SetDiscoverable(true);
+
+			TSharedPtr<bool> CbResult = MakeShared<bool>(false);
+
+			FCaptureSourceFactory::FDiscoverCallback Cb =
+				FCaptureSourceFactory::FDiscoverCallback::CreateLambda(
+					[this, CbResult](TArray<FCaptureSourceFactory::FCaptureSourceResult> InResults)
+			{
+				check(InResults.Num() == 1);
+				const FCaptureSourceFactory::FCaptureSourceResult& Result = InResults[0];
+
+				if (!TestTrue("Capture source should be discovered (no error)", Result.IsValid()))
+				{
+					*CbResult = false;
+				}
+
+				*CbResult = true;
+			});
+
+			FCaptureVoidResult Result = CaptureSourceFactory->Discover(MoveTemp(Cb));
+
+			UTEST_TRUE("Starting discovery should pass", Result.HasValue());
+			UTEST_TRUE("Result from callback should be true", *CbResult);
+
+			return true;
+		});
+
+		It("returns an error as factory isn't discoverable", [this]()
+		{
+			CaptureSourceFactory->SetDiscoverable(false);
+			FCaptureVoidResult Result = CaptureSourceFactory->Discover(FCaptureSourceFactory::FDiscoverCallback());
+
+			UTEST_TRUE("Result should contain an error", Result.HasError());
+			UTEST_EQUAL("Resulting error should be 'Not Supported'", Result.GetError().GetErrorCode(), ECaptureSourceError::NotSupported);
+
+			return true;
+		});	
+	});
+
+	Describe("CreateCaptureSource()", [this]()
+	{
+		It("returns a valid capture source", [this]()
+		{
+			FString Name = TEXT("CaptureSourceName");
+			FString CreationParamValue = TEXT("Value");
+
+			TMap<FString, FPropertyValue> CreationParamsValue;
+			CreationParamsValue.Add(FTestCaptureSourceFactory::CreationParam, MakePropertyValue<FString>(CreationParamValue));
+			
+			FCaptureSourceFactory::FCaptureSourceResult Result = CaptureSourceFactory->CreateCaptureSource(Name, MoveTemp(CreationParamsValue));
+
+			UTEST_TRUE("Result should be valid", Result.HasValue());
+
+			TUniquePtr<FCaptureSource> CaptureSource = Result.StealValue();
+
+			UTEST_TRUE("Result should contain a valid capture source", CaptureSource.IsValid());
+
+			FTestCaptureSource* TestCaptureSource = static_cast<FTestCaptureSource*>(CaptureSource.Get());
+			UTEST_EQUAL("Capture source should contain valid creation param", TestCaptureSource->ParamValue, CreationParamValue);
+
+			return true;
+		});
+
+		It("returns an error because parameters are missing", [this]()
+		{
+			FString Name = TEXT("CaptureSourceName");
+
+			TMap<FString, FPropertyValue> CreationParamsValue; // Empty
+
+			FCaptureSourceFactory::FCaptureSourceResult Result = CaptureSourceFactory->CreateCaptureSource(Name, MoveTemp(CreationParamsValue));
+
+			UTEST_TRUE("Result should contain an error", Result.HasError());
+
+			return true;
+		});
+	});
+
+	Describe("CreateCaptureSourceDescriptor()", [this]()
+	{
+		It("returns a valid descriptor containing one creation param", [this]()
+		{
+			FCaptureSourceDescriptor Descriptor = CaptureSourceFactory->CreateCaptureSourceDescriptor();
+
+			const TArray<FPropertyDesc>& CreationParams = Descriptor.GetCreationParams();
+
+			UTEST_EQUAL("Creation params should contain 1 param", CreationParams.Num(), 1);
+
+			UTEST_EQUAL("Creation param name should match", CreationParams[0].Name, FTestCaptureSourceFactory::CreationParam);
+			UTEST_EQUAL("Creation param type should match", CreationParams[0].Type, FPropertyDesc::EType::String);
+
+			return true;
+		});
+	});
+}
+
+#endif // WITH_DEV_AUTOMATION_TESTS
\ No newline at end of file
diff --git a/Engine/Plugins/VirtualProduction/CaptureManager/Source/CaptureSourceFramework/Private/Tests/TestCaptureSourceManager.cpp b/Engine/Plugins/VirtualProduction/CaptureManager/Source/CaptureSourceFramework/Private/Tests/TestCaptureSourceManager.cpp
new file mode 100644
--- /dev/null
+++ b/Engine/Plugins/VirtualProduction/CaptureManager/Source/CaptureSourceFramework/Private/Tests/TestCaptureSourceManager.cpp
@@ -0,0 +1,264 @@
+// Copyright Epic Games, Inc. All Rights Reserved.
+
+#include "CaptureSourceManager.h"
+
+#include "Misc/AutomationTest.h"
+
+#if WITH_DEV_AUTOMATION_TESTS
+
+BEGIN_DEFINE_SPEC(FCaptureSourceManagerTest, "Plugin.CaptureSourceFramework.CaptureSourceManager", EAutomationTestFlags::EditorContext | EAutomationTestFlags::ProductFilter | EAutomationTestFlags::MediumPriority)
+
+class FTestCaptureSource final : public FCaptureSource
+{
+public:
+	FTestCaptureSource(const FString& InName)
+		: FCaptureSource(FTestCaptureSourceFactory::Id, InName)
+	{
+	}
+
+	virtual FCaptureVoidResult Start() final
+	{
+		return MakeValue();
+	}
+
+	virtual FCaptureVoidResult Stop() final
+	{
+		return MakeValue();
+	}
+};
+
+class FTestCaptureSourceFactory final : public FCaptureSourceFactory
+{
+public:
+
+	static const FString Id;
+
+	static const FString CreationParam;
+
+	virtual FString GetCaptureSourceFactoryName() const override
+	{
+		return TEXT("Test Capture Source Factory");
+	}
+
+	virtual FString GetCaptureSourceFactoryId() const override
+	{
+		return Id;
+	}
+
+	virtual bool IsDiscoverable() const override
+	{
+		return bIsDiscoverable;
+	}
+
+	void SetDiscoverable(bool bInIsDiscoverable)
+	{
+		bIsDiscoverable = bInIsDiscoverable;
+	}
+
+private:
+
+	virtual void DiscoverImpl(FDiscoverCallback InCallback, FCaptureVoidResult& OutResult) override
+	{
+		TArray<FCaptureSourceResult> DiscoveredCaptureSources;
+
+		DiscoveredCaptureSources.Add(MakeValue(MakeUnique<FTestCaptureSource>(TEXT("DefaultName"))));
+
+		InCallback.ExecuteIfBound(MoveTemp(DiscoveredCaptureSources));
+
+		OutResult = MakeValue();
+	}
+
+	virtual FCaptureSourceResult CreateCaptureSourceImpl(const FString& InName, const TMap<FString, FPropertyValue>& InCreationParamsValue) override
+	{
+		if (!InCreationParamsValue.IsEmpty())
+		{
+			return MakeError(FCaptureSourceError(ECaptureSourceError::InvalidArgument));
+		}
+
+		return MakeValue(MakeUnique<FTestCaptureSource>(InName));
+	}
+
+	virtual void CreateCaptureSourceDescriptorImpl(FCaptureSourceDescriptor& OutDescriptor) const override
+	{
+		OutDescriptor.SetDiscoverable(bIsDiscoverable);
+	}
+
+	bool bIsDiscoverable;
+};
+
+FCaptureSourceManager Manager;
+FTestCaptureSourceFactory* FactoryPtr = nullptr;
+
+using FCreationParamValues = TMap<FString, FPropertyValue>;
+
+END_DEFINE_SPEC(FCaptureSourceManagerTest)
+
+const FString FCaptureSourceManagerTest::FTestCaptureSourceFactory::Id = TEXT("TestCaptureSourceFactoryId");
+
+void FCaptureSourceManagerTest::Define()
+{
+	TUniquePtr<FTestCaptureSourceFactory> Factory = MakeUnique<FTestCaptureSourceFactory>();
+	FactoryPtr = Factory.Get();
+
+	Manager.RegisterCaptureSourceFactory(MoveTemp(Factory));
+	bSuppressLogs = true;
+
+	Describe("CreateCaptureSource()", [this]()
+	{
+		It("returns a valid capture source", [this]()
+		{
+			FString Name = TEXT("Name");
+			FCaptureSourceManager::FCaptureSourceResult Result = Manager.CreateCaptureSource(FTestCaptureSourceFactory::Id, MoveTemp(Name), FCreationParamValues());
+
+			UTEST_TRUE("Result should be valid", Result.HasValue());
+
+			return true;
+		});
+
+		It("returns an error as invalid capture source factory is provided", [this]()
+		{
+			FString Name = TEXT("Name");
+			FString InvalidCaptureSourceFactory = TEXT("Invalid");
+			FCaptureSourceManager::FCaptureSourceResult Result = Manager.CreateCaptureSource(MoveTemp(InvalidCaptureSourceFactory), MoveTemp(Name), FCreationParamValues());
+
+			UTEST_TRUE("Result should contain an error", Result.HasError());
+
+			return true;
+		});
+
+		It("returns an error as number of arguments is invalid", [this]()
+		{
+			FString Name = TEXT("Name");
+
+			FCreationParamValues ParamValues;
+			ParamValues.Add(TEXT("Param"), MakePropertyValue<FString>(TEXT("Value")));
+
+			FCaptureSourceManager::FCaptureSourceResult Result = Manager.CreateCaptureSource(FTestCaptureSourceFactory::Id, MoveTemp(Name), MoveTemp(ParamValues));
+
+			UTEST_TRUE("Result should contain an error", Result.HasError());
+
+			return true;
+		});
+	});
+
+	Describe("DiscoverCaptureSources()", [this]()
+	{
+		It("runs discovery without problems as factory is discoverable", [this]()
+		{
+			FactoryPtr->SetDiscoverable(true);
+
+			TSharedPtr<bool> CbResult = MakeShared<bool>(false);
+
+			FCaptureSourceManager::FDiscoverCallback Cb = 
+				FCaptureSourceManager::FDiscoverCallback::CreateLambda([this, CbResult](TArray<FCaptureSourceManager::FCaptureSourceResult> InResults)
+			{
+				if (!TestTrue("Result contains a single valid source", InResults[0].IsValid()))
+				{
+					*CbResult = false;
+				}
+
+				*CbResult = true;
+			});
+
+			Manager.DiscoverCaptureSources(MoveTemp(Cb));
+
+			UTEST_TRUE("Result from callback should be true", *CbResult);
+
+			return true;
+		});
+
+		It("returns an error as factory isn't discoverable", [this]()
+		{
+			FactoryPtr->SetDiscoverable(false);
+
+			TSharedPtr<bool> CbResult = MakeShared<bool>(false);
+
+			FCaptureSourceManager::FDiscoverCallback Cb =
+				FCaptureSourceManager::FDiscoverCallback::CreateLambda([this, CbResult](TArray<FCaptureSourceManager::FCaptureSourceResult> InResults)
+			{
+				if (!TestTrue("Result contains an error", InResults.IsEmpty()))
+				{
+					*CbResult = false;
+				}
+
+				*CbResult = true;
+			});
+
+			Manager.DiscoverCaptureSources(MoveTemp(Cb));
+
+			UTEST_TRUE("Result from callback should be true", *CbResult);
+
+			return true;
+		});
+	});
+
+	Describe("DiscoverCaptureSourcesForFactory()", [this]()
+	{
+		It("runs discovery without problems as factory is discoverable", [this]()
+		{
+			FactoryPtr->SetDiscoverable(true);
+
+			TSharedPtr<bool> CbResult = MakeShared<bool>(false);
+
+			FCaptureSourceManager::FDiscoverCallback Cb =
+				FCaptureSourceManager::FDiscoverCallback::CreateLambda([this, CbResult](TArray<FCaptureSourceManager::FCaptureSourceResult> InResults)
+			{
+				if (!TestTrue("Result contains a single valid source", InResults[0].IsValid()))
+				{
+					*CbResult = false;
+				}
+
+				*CbResult = true;
+			});
+
+			Manager.DiscoverCaptureSourcesForFactory(FTestCaptureSourceFactory::Id, MoveTemp(Cb));
+
+			UTEST_TRUE("Result from callback should be true", *CbResult);
+
+			return true;
+		});
+
+		It("returns an error as factory isn't discoverable", [this]()
+		{
+			FactoryPtr->SetDiscoverable(false);
+
+			TSharedPtr<bool> CbResult = MakeShared<bool>(false);
+
+			FCaptureSourceManager::FDiscoverCallback Cb =
+				FCaptureSourceManager::FDiscoverCallback::CreateLambda([this, CbResult](TArray<FCaptureSourceManager::FCaptureSourceResult> InResults)
+			{
+				if (!TestTrue("Result contains an error", InResults.IsEmpty()))
+				{
+					*CbResult = false;
+				}
+
+				*CbResult = true;
+			});
+
+			Manager.DiscoverCaptureSourcesForFactory(FTestCaptureSourceFactory::Id, MoveTemp(Cb));
+
+			UTEST_TRUE("Result from callback should be true", *CbResult);
+
+			return true;
+		});
+	});
+
+	Describe("RemoveCaptureSource()", [this]()
+	{
+		It("removes all created capture sources", [this]()
+		{
+			TArray<FCaptureSourceId> Ids = Manager.GetCaptureSourceIds();
+
+			for (FCaptureSourceId Id : Ids)
+			{
+				Manager.RemoveCaptureSource(Id);
+			}
+
+			UTEST_TRUE("Array of capture sources should be empty after removal", Manager.GetCaptureSourceIds().IsEmpty());
+
+			return true;
+		});
+	});
+}
+
+#endif // WITH_DEV_AUTOMATION_TESTS
\ No newline at end of file
diff --git a/Engine/Plugins/VirtualProduction/CaptureManager/Source/CaptureSourceFramework/Public/Capabilities/DiskCapacityCapability.h b/Engine/Plugins/VirtualProduction/CaptureManager/Source/CaptureSourceFramework/Public/Capabilities/DiskCapacityCapability.h
new file mode 100644
--- /dev/null
+++ b/Engine/Plugins/VirtualProduction/CaptureManager/Source/CaptureSourceFramework/Public/Capabilities/DiskCapacityCapability.h
@@ -0,0 +1,70 @@
+// Copyright Epic Games, Inc. All Rights Reserved.
+
+#pragma once
+
+#include "CaptureSourceCapability.h"
+
+struct CAPTURESOURCEFRAMEWORK_API FStateUpdatedEvent : public FCaptureEvent
+{
+	static const FString Name;
+
+	FStateUpdatedEvent(uint64 InTotalSpace,
+					   uint64 InRemainingSpace)
+		: FCaptureEvent(Name)
+		, TotalSpace(InTotalSpace)
+		, RemainingSpace(InRemainingSpace)
+		, UsedSpace(InTotalSpace - InRemainingSpace)
+	{
+	}
+
+	uint64 TotalSpace;
+	uint64 RemainingSpace;
+	uint64 UsedSpace;
+};
+
+class CAPTURESOURCEFRAMEWORK_API FUpdateCmd : public FCommandBase
+{
+public:
+
+	static const FString Name;
+
+	FUpdateCmd();
+};
+
+class CAPTURESOURCEFRAMEWORK_API FDiskCapacityCapabilityBase : public FCaptureSourceCapability
+{
+public:
+	static const FString Name;
+	static const FString TotalSpace;
+	static const FString RemainingSpace;
+	static const FString UsedSpace;
+
+	FDiskCapacityCapabilityBase();
+	virtual ~FDiskCapacityCapabilityBase() = default;
+
+	virtual uint64 GetTotalSpace() = 0;
+	virtual uint64 GetRemainingSpace() = 0;
+	virtual uint64 GetUsedSpace() = 0;
+
+	virtual void Update() = 0;
+};
+
+class CAPTURESOURCEFRAMEWORK_API FDiskCapacityCapability : public FDiskCapacityCapabilityBase
+{
+public:
+
+	FDiskCapacityCapability();
+	virtual ~FDiskCapacityCapability() = default;
+
+	virtual uint64 GetTotalSpace() override;
+	virtual uint64 GetRemainingSpace() override;
+	virtual uint64 GetUsedSpace() override;
+
+	virtual void Update() override;
+
+	void SetTotalSpacePropertyHandler(FPropertyGetter InGetter);
+	void SetRemainingSpacePropertyHandler(FPropertyGetter InGetter);
+	void SetUpdateCmdHandler(FCommandHandler InCommandHandler);
+
+	void PublishStateUpdatedEvent(uint64 InTotalBytes, uint64 InRemainingBytes);
+};
\ No newline at end of file
diff --git a/Engine/Plugins/VirtualProduction/CaptureManager/Source/CaptureSourceFramework/Public/Capabilities/IngestCapability.h b/Engine/Plugins/VirtualProduction/CaptureManager/Source/CaptureSourceFramework/Public/Capabilities/IngestCapability.h
new file mode 100644
--- /dev/null
+++ b/Engine/Plugins/VirtualProduction/CaptureManager/Source/CaptureSourceFramework/Public/Capabilities/IngestCapability.h
@@ -0,0 +1,149 @@
+// Copyright Epic Games, Inc. All Rights Reserved.
+
+#pragma once
+
+#include "CaptureSourceCapability.h"
+
+#include "Templates/ValueOrError.h"
+
+class CAPTURESOURCEFRAMEWORK_API FIngestCapabilityError
+{
+public:
+
+	enum ECode
+	{
+		Ok = 0,
+		AbortedByUser = 1,
+		InternalError = 2
+	};
+
+	FIngestCapabilityError(ECode InCode, FString InMessage);
+
+	ECode GetCode() const;
+	const FString& GetMessage() const;
+
+private:
+
+	ECode Code;
+	FString Message;
+};
+
+using FTakeId = int32;
+
+struct FTakeInfo
+{
+	FTakeId Id;
+
+	FString SlateName;
+	int16 TakeNumber;
+
+	int32 NumFrames;
+	double FrameRate;
+
+	FDateTime Date;
+
+	FString GetName()
+	{
+		return FString::Format(TEXT("{0}_{1}"), { SlateName, TakeNumber });
+	}
+};
+
+class CAPTURESOURCEFRAMEWORK_API FImportTakeCmd : public FCommandBase
+{
+public:
+	static const FString Name;
+
+	DECLARE_DELEGATE_TwoParams(FProcessProgressReporter, FTakeId, double);
+	using FProcessFinished = TDelegate<void(FTakeId, TValueOrError<void, FIngestCapabilityError>)>;
+
+	FImportTakeCmd(FProcessProgressReporter InProgressReporter, 
+				   FProcessFinished InProcessFinished);
+
+	FProcessProgressReporter ProgressReporter;
+	FProcessFinished ProcessFinished;
+};
+
+class CAPTURESOURCEFRAMEWORK_API FUpdateTakeListCmd : public FCommandBase
+{
+public:
+
+	static const FString Name;
+
+	DECLARE_DELEGATE_OneParam(FCallback, TArray<FTakeId>);
+
+	FUpdateTakeListCmd(FCallback InCallback);
+
+	FCallback Callback;
+};
+
+class CAPTURESOURCEFRAMEWORK_API FIngestCapabilityBase : public FCaptureSourceCapability
+{
+public:
+	class CAPTURESOURCEFRAMEWORK_API FTakeProcessHandle
+	{
+	public:
+
+		FTakeProcessHandle() = default;
+		FTakeProcessHandle(FIngestCapabilityBase* InOwner, FTakeId InId);
+
+		void Cancel();
+
+	private:
+
+		FIngestCapabilityBase* Owner;
+		FTakeId Id;
+	};
+
+	static const FString Name;
+
+	static const FString CancelCmd;
+
+	static const FString Id;
+	static const FString Location;
+
+	FIngestCapabilityBase();
+
+	virtual ~FIngestCapabilityBase() = default;
+
+	virtual FTakeProcessHandle ImportTake(FTakeId InTakeId,
+										  FString InLocation,
+										  FImportTakeCmd::FProcessFinished InProcessFinished,
+										  FImportTakeCmd::FProcessProgressReporter InProcessProgressReported) = 0;
+
+	virtual void UpdateTakeList(FUpdateTakeListCmd::FCallback InUpdateTakeListCallback) = 0;
+
+	FTakeId AddTakeInfo(FTakeInfo InTakeInfo);
+	FTakeInfo GetTakeInfo(FTakeId InTakeId) const;
+	TArray<FTakeId> GetTakeIds() const;
+
+protected:
+
+	virtual void Cancel(FTakeId InId) = 0;
+
+	TMap<FTakeId, FTakeInfo> Takes;
+
+private:
+
+	std::atomic<FTakeId> CurrentTakeId;
+};
+
+class CAPTURESOURCEFRAMEWORK_API FIngestCapability : public FIngestCapabilityBase
+{
+public:
+
+	FIngestCapability();
+
+	FTakeProcessHandle ImportTake(FTakeId InTakeId,
+								  FString InLocation,
+								  FImportTakeCmd::FProcessFinished InProcessFinished,
+								  FImportTakeCmd::FProcessProgressReporter InProcessProgressReported) override;
+	void UpdateTakeList(FUpdateTakeListCmd::FCallback InUpdateTakeListCallback) override;
+
+	void SetImportTakeHandler(FCommandHandler InCommandHandler);
+	void SetUpdateTakeListHandler(FCommandHandler InCommandHandler);
+	void SetCancelHandler(FCommandHandler InCommandHandler);
+
+private:
+
+	void Cancel(FTakeId InTakeId);
+};
\ No newline at end of file
diff --git a/Engine/Plugins/VirtualProduction/CaptureManager/Source/CaptureSourceFramework/Public/Capabilities/RecordingCapability.h b/Engine/Plugins/VirtualProduction/CaptureManager/Source/CaptureSourceFramework/Public/Capabilities/RecordingCapability.h
new file mode 100644
--- /dev/null
+++ b/Engine/Plugins/VirtualProduction/CaptureManager/Source/CaptureSourceFramework/Public/Capabilities/RecordingCapability.h
@@ -0,0 +1,52 @@
+// Copyright Epic Games, Inc. All Rights Reserved.
+
+#pragma once
+
+#include "CaptureSourceCapability.h"
+
+struct CAPTURESOURCEFRAMEWORK_API FRecordingEvent : public FCaptureEvent
+{
+	static const FString Name;
+
+	FRecordingEvent(bool bInIsRecording)
+		: FCaptureEvent(Name)
+		, bIsRecording(bInIsRecording)
+	{
+	}
+
+	bool bIsRecording;
+};
+
+class CAPTURESOURCEFRAMEWORK_API FRecordingCapabilityBase : public FCaptureSourceCapability
+{
+public:
+	static const FString Name;
+	static const FString SlateName;
+	static const FString TakeNumber;
+	static const FString StartRecordingCmd;
+	static const FString StopRecordingCmd;
+
+	FRecordingCapabilityBase();
+	virtual ~FRecordingCapabilityBase() = default;
+
+	virtual void StartRecording(FString InSlateName, int64 InTakeNumber) = 0;
+	virtual void StopRecording() = 0;
+
+	virtual void SetStartRecordingHandler(FCommandHandler InCommandHandler) = 0;
+	virtual void SetStopRecordingHandler(FCommandHandler InCommandHandler) = 0;
+};
+
+class CAPTURESOURCEFRAMEWORK_API FRecordingCapability : public FRecordingCapabilityBase
+{
+public:
+
+	FRecordingCapability();
+
+	void StartRecording(FString InSlateName, int64 InTakeNumber);
+	void StopRecording();
+
+	void SetStartRecordingHandler(FCommandHandler InCommandHandler);
+	void SetStopRecordingHandler(FCommandHandler InCommandHandler);
+
+	void PublishRecordingEvent(bool bIsRecording);
+};
\ No newline at end of file
diff --git a/Engine/Plugins/VirtualProduction/CaptureManager/Source/CaptureSourceFramework/Public/CaptureSource.h b/Engine/Plugins/VirtualProduction/CaptureManager/Source/CaptureSourceFramework/Public/CaptureSource.h
new file mode 100644
--- /dev/null
+++ b/Engine/Plugins/VirtualProduction/CaptureManager/Source/CaptureSourceFramework/Public/CaptureSource.h
@@ -0,0 +1,78 @@
+// Copyright Epic Games, Inc. All Rights Reserved.
+
+#pragma once
+
+#include "CaptureSourceCapability.h"
+
+#include "CaptureSourceError.h"
+
+struct CAPTURESOURCEFRAMEWORK_API FReachableEvent : public FCaptureEvent
+{
+	static const FString Name;
+
+	FReachableEvent(bool bInIsReachable)
+		: FCaptureEvent(Name)
+		, bIsReachable(bInIsReachable)
+	{
+	}
+
+	bool bIsReachable;
+};
+
+namespace UE::CaptureManager::Private
+{
+class ICaptureSource
+{
+public:
+	
+	using FCapabilityMap = TMap<FString, FCaptureSourceCapability*>;
+
+	virtual ~ICaptureSource() = default;
+
+	virtual FCaptureVoidResult Start() = 0;
+	virtual FCaptureVoidResult Stop() = 0;
+
+	virtual FCaptureSourceCapability* GetCapability(FString InName) const = 0;
+	virtual TArray<FString> GetCapabilities() const = 0;
+
+	virtual const FString& GetCaptureSourceFactoryId() const = 0;
+	virtual const FString& GetName() const = 0;
+
+	virtual const TMap<FString, FPropertyValue>& GetCreationParams() const = 0;
+};
+}
+
+class CAPTURESOURCEFRAMEWORK_API FCaptureSource : 
+	public UE::CaptureManager::Private::ICaptureSource,
+	public FCaptureEventSource
+{
+public:
+
+	FCaptureSource(const FString& InCaptureSourceFactoryId, const FString& InName);
+	virtual ~FCaptureSource() = default;
+
+	virtual FCaptureSourceCapability* GetCapability(FString InName) const final;
+	virtual TArray<FString> GetCapabilities() const final;
+
+	virtual const FString& GetCaptureSourceFactoryId() const final;
+	virtual const FString& GetName() const final;
+
+	void SetCreationParams(TMap<FString, FPropertyValue> InCreationParams);
+	virtual const TMap<FString, FPropertyValue>& GetCreationParams() const final;
+
+protected:
+
+	void AddCapability(FCaptureSourceCapability* InCapability);
+
+private:
+
+	FCaptureSource(const FCaptureSource& InOther) = delete;
+	FCaptureSource& operator=(const FCaptureSource& InOther) = delete;
+
+	const FString CaptureSourceFactoryId;
+	FString Name;
+
+	FCapabilityMap CapabilityMap;
+
+	TMap<FString, FPropertyValue> CreationParams;
+};
\ No newline at end of file
diff --git a/Engine/Plugins/VirtualProduction/CaptureManager/Source/CaptureSourceFramework/Public/CaptureSourceCapability.h b/Engine/Plugins/VirtualProduction/CaptureManager/Source/CaptureSourceFramework/Public/CaptureSourceCapability.h
new file mode 100644
--- /dev/null
+++ b/Engine/Plugins/VirtualProduction/CaptureManager/Source/CaptureSourceFramework/Public/CaptureSourceCapability.h
@@ -0,0 +1,75 @@
+// Copyright Epic Games, Inc. All Rights Reserved.
+
+#pragma once
+
+#include "Async/EventSourceUtils.h"
+
+#include "CaptureSourceCapabilityProperty.h"
+#include "CaptureSourceCapabilityCommand.h"
+
+struct FCaptureSourceCapabilityDesc
+{
+	FString Name;
+	TArray<FPropertyDesc> Properties;
+	TArray<FCommandDesc> Commands;
+
+	FCaptureSourceCapabilityDesc(FString InName, TArray<FPropertyDesc> InProperties, TArray<FCommandDesc> InCommands);
+};
+
+struct CAPTURESOURCEFRAMEWORK_API FCapturePropertyChangedEvent : public FCaptureEvent
+{
+public:
+	static const FString& EventName;
+
+	FCapturePropertyChangedEvent(const FString& InName, FPropertyValue InValue);
+
+	const FString Name;
+	const FPropertyValue Value;
+};
+
+class CAPTURESOURCEFRAMEWORK_API FCaptureSourceCapability : public FCaptureEventSource
+{
+public:
+
+	DECLARE_DELEGATE_TwoParams(FPropertySetter, const FString&, FPropertyValue);
+	DECLARE_DELEGATE_RetVal_OneParam(FPropertyValue, FPropertyGetter, const FString&);
+	DECLARE_DELEGATE_OneParam(FCommandHandler, TSharedPtr<FCommandBase>);
+
+	virtual ~FCaptureSourceCapability() = default;
+
+	void PublishPropertyChangedEvent(const FString& InName, FPropertyValue InValue);
+
+	void SetPropertyValue(const FString& InName, FPropertyValue InValue);
+	FPropertyValue GetPropertyValue(const FString& InName);
+
+	void ExecuteCommand(TSharedPtr<FCommandBase> InCommand);
+
+	TArray<FPropertyDesc> GetProperties() const;
+	TArray<FCommandDesc> GetCommands() const;
+
+	FString GetName() const;
+
+	FCaptureSourceCapabilityDesc CreateDescriptor();
+
+protected:
+
+	FCaptureSourceCapability(FString InName);
+
+	void SetPropertySetter(const FString& InName, FPropertySetter InSetter);
+	void SetPropertyGetter(const FString& InName, FPropertyGetter InGetter);
+	void SetCommandHandler(const FString& InName, FCommandHandler InCommandHandler);
+
+	void AddProperty(FPropertyDesc InPropertyDesc);
+	void AddCommand(FCommandDesc InCommand);
+
+private:
+
+	FPropertyList Properties;
+	FCommandList Commands;
+
+	FString Name;
+
+	TMap<FString, FPropertySetter> Setters;
+	TMap<FString, FPropertyGetter> Getters;
+	TMap<FString, FCommandHandler> CommandHandlers;
+};
diff --git a/Engine/Plugins/VirtualProduction/CaptureManager/Source/CaptureSourceFramework/Public/CaptureSourceCapabilityCommand.h b/Engine/Plugins/VirtualProduction/CaptureManager/Source/CaptureSourceFramework/Public/CaptureSourceCapabilityCommand.h
new file mode 100644
--- /dev/null
+++ b/Engine/Plugins/VirtualProduction/CaptureManager/Source/CaptureSourceFramework/Public/CaptureSourceCapabilityCommand.h
@@ -0,0 +1,56 @@
+// Copyright Epic Games, Inc. All Rights Reserved.
+
+#pragma once
+
+#include "CaptureSourceCapabilityProperty.h"
+
+#include "Templates/SharedPointer.h"
+
+struct CAPTURESOURCEFRAMEWORK_API FCommandDesc
+{
+	FString Name;
+	TArray<FPropertyDesc> Params;
+
+	FCommandDesc(FString InName, TArray<FPropertyDesc> InParams);
+};
+
+class CAPTURESOURCEFRAMEWORK_API FCommandBase
+{
+public:
+	using FValues = TMap<FString, FPropertyValue>;
+
+	FCommandBase(FString InName);
+	virtual ~FCommandBase() = default;
+
+	void AddParamValue(const FString& InName, FPropertyValue InValue);
+	FPropertyValue GetParamValue(const FString& InName) const;
+
+	void SetParamValues(FValues InValues);
+
+	const FValues& GetParamValues() const;
+	bool ContainsParamValue(const FString& InName) const;
+
+	const FString& GetName() const;
+
+private:
+
+	FString Name;
+	FValues Values;
+};
+
+class FCommandList
+{
+public:
+
+	FCommandList() = default;
+
+	void AddCommand(FCommandDesc InCommandDesc);
+	TArray<FCommandDesc> GetCommands() const;
+
+	void CheckCommandExists(const FString& InName);
+	void CheckCommandParamValue(const FString& InName, const TSharedPtr<FCommandBase>& InCommand);
+
+private:
+
+	TArray<FCommandDesc> Commands;
+};
\ No newline at end of file
diff --git a/Engine/Plugins/VirtualProduction/CaptureManager/Source/CaptureSourceFramework/Public/CaptureSourceCapabilityProperty.h b/Engine/Plugins/VirtualProduction/CaptureManager/Source/CaptureSourceFramework/Public/CaptureSourceCapabilityProperty.h
new file mode 100644
--- /dev/null
+++ b/Engine/Plugins/VirtualProduction/CaptureManager/Source/CaptureSourceFramework/Public/CaptureSourceCapabilityProperty.h
@@ -0,0 +1,65 @@
+// Copyright Epic Games, Inc. All Rights Reserved.
+
+#pragma once
+
+#include "Containers/Array.h"
+#include "Containers/UnrealString.h"
+
+#include "Misc/TVariant.h"
+
+using FPropertyValue = TVariant<FEmptyVariantState, bool, int64, double, FString, TArray<bool>, TArray<int64>, TArray<double>, TArray<FString>>;
+
+struct CAPTURESOURCEFRAMEWORK_API FPropertyDesc
+{
+	enum class EType
+	{
+		Bool = 0,
+		Number,
+		FloatingPoint,
+		String,
+		ArrayBool,
+		ArrayNumber,
+		ArrayString,
+		ArrayFloatingPoint
+	};
+
+	enum class EAccess
+	{
+		ReadOnly = 0,
+		ReadWrite
+	};
+
+	FString Name;
+	EType Type;
+	EAccess Access;
+
+	FPropertyDesc(FString InName, EType InType, EAccess InAccess);
+	FPropertyDesc(FString InName, EType InType);
+};
+
+template<typename T>
+FPropertyValue MakePropertyValue(T InValue)
+{
+	return FPropertyValue(TInPlaceType<T>(), MoveTemp(InValue));
+}
+
+class FPropertyList
+{
+public:
+
+	FPropertyList() = default;
+
+	void AddProperty(FPropertyDesc InPropertyDesc);
+	TArray<FPropertyDesc> GetProperties() const;
+
+	void CheckPropertyExists(const FString& InName);
+	void CheckPropertyValueTypeMatch(const FString& InName, const FPropertyValue& InValue);
+	void CheckPropertyAllowsWrite(const FString& InName);
+
+	static bool CheckPropertyValueType(FPropertyDesc::EType InType, const FPropertyValue& InValue);
+
+private:
+
+	TArray<FPropertyDesc> Properties;
+};
+
diff --git a/Engine/Plugins/VirtualProduction/CaptureManager/Source/CaptureSourceFramework/Public/CaptureSourceDetailsCustomization.h b/Engine/Plugins/VirtualProduction/CaptureManager/Source/CaptureSourceFramework/Public/CaptureSourceDetailsCustomization.h
new file mode 100644
--- /dev/null
+++ b/Engine/Plugins/VirtualProduction/CaptureManager/Source/CaptureSourceFramework/Public/CaptureSourceDetailsCustomization.h
@@ -0,0 +1,112 @@
+// Copyright Epic Games, Inc. All Rights Reserved.
+
+#pragma once
+
+#include "CaptureSourceCapability.h"
+#include "CaptureSourceFactory.h"
+#include "CaptureSourceManager.h"
+
+#include "Widgets/SWidget.h"
+
+class CAPTURESOURCEFRAMEWORK_API FCaptureSourceDetailsCustomization
+{
+public:
+
+	DECLARE_DELEGATE_OneParam(FCaptureSourceCreated, FCaptureSourceId);
+
+	DECLARE_DELEGATE_RetVal_ThreeParams(TSharedPtr<SWidget>, FPropertyWidgetCreator, FCaptureSourceId, FCaptureSourceCapability*, const FPropertyDesc&);
+	DECLARE_DELEGATE_RetVal_ThreeParams(TSharedPtr<SWidget>, FCommandWidgetCreator, FCaptureSourceId, FCaptureSourceCapability*, const FCommandDesc&);
+	DECLARE_DELEGATE_RetVal_TwoParams(TSharedPtr<SWidget>, FCapabilityWidgetCreator, FCaptureSourceId, FCaptureSourceCapability*);
+	DECLARE_DELEGATE_RetVal_TwoParams(TSharedPtr<SWindow>, FCaptureSourceWidgetCreator, const FCaptureSourceDescriptor&, FCaptureSourceCreated);
+
+	FCaptureSourceDetailsCustomization();
+
+	void RegisterPropertyCustomization(const FString& InCaptureSourceClassId,
+									   const FString& InCapabilityName, 
+									   const FString& InPropertyName, 
+									   FPropertyWidgetCreator InWidgetCreator);
+
+	void RegisterCommandCustomization(const FString& InCaptureSourceClassId,
+									  const FString& InCapabilityName,
+									  const FString& InCommandName, 
+									  FCommandWidgetCreator InWidgetCreator);
+
+	void RegisterCapabilityCustomization(const FString& InCaptureSourceClassId,
+										 const FString& InCapabilityName, 
+										 FCapabilityWidgetCreator InWidgetCreator);
+
+	void RegisterCaptureSourceCustomization(const FString& InCaptureSourceClassId, 
+											FCaptureSourceWidgetCreator InWidgetCreator);
+
+	TSharedPtr<SWidget> CreatePropertyWidget(FCaptureSourceId InCaptureSourceId,
+											 FCaptureSourceCapability* InCapability, 
+											 const FPropertyDesc& InProperty);
+
+	TSharedPtr<SWidget> CreateCommandWidget(FCaptureSourceId InCaptureSourceId,
+											FCaptureSourceCapability* InCapability,
+											const FCommandDesc& InCommand);
+
+	TSharedPtr<SWidget> CreateCapabilityWidget(FCaptureSourceId InCaptureSourceId,
+											   FCaptureSourceCapability* InCapability);
+
+	TSharedPtr<SWindow> CreateCaptureSourceWidget(const FCaptureSourceDescriptor& InCreationParams,
+												  FCaptureSourceCreated InCaptureSourceCreated);
+	
+private:
+
+	template <typename ... Args>
+	FString GenerateKey(Args... InArgs)
+	{
+		static_assert(sizeof...(Args) != 0 && sizeof...(Args) <= 3,
+					  "Method currently supports 1 - 3 arguments");
+
+		FString Format;
+		if constexpr (sizeof...(Args) == 1)
+		{
+			Format = TEXT("{0}");
+		}
+		else if constexpr (sizeof...(Args) == 2)
+		{
+			Format = TEXT("{0}.{1}");
+		}
+		else
+		{
+			Format = TEXT("{0}.{1}.{2}");
+		}
+		return FString::Format(*Format, { (InArgs)... });
+	}
+
+	FString GenerateDefaultMapKey(const FString& InName);
+
+	bool ContainsPropertyWidgetCustomization(const FString& InFactoryId, 
+											 const FString& InCapabilityName, 
+											 const FString& InPropertyName);
+
+	TSharedPtr<SWidget> CreateDefaultPropertyWidget(FCaptureSourceCapability* InCapability,
+													const FPropertyDesc& InProperty);
+
+	TSharedPtr<SWidget> CreateDefaultCapabilityPropertyWidget(FCaptureSourceCapability* InCapability, 
+															  const FPropertyDesc& InProperty,
+															  TSharedPtr<FPropertyValue> InPropertyValue);
+
+	TSharedPtr<SWidget> CreateParamWidget(const FPropertyDesc& InProperty, FPropertyValue& InPropertyValue);
+
+	TSharedPtr<SWidget> CreateDefaultCommandWidget(FCaptureSourceCapability* InCapability, const FCommandDesc& InCommand);
+
+	TSharedPtr<SWidget> MakeKeyWidget(const FPropertyDesc& InProperty);
+	TSharedPtr<SWidget> MakeValueWidget(FCaptureSourceCapability* InCapability, 
+										const FPropertyDesc& InProperty,
+										TSharedPtr<FPropertyValue> InPropertyValue);
+
+	TSharedPtr<SWidget> MakeValueWidget(const FPropertyDesc& InProperty, FPropertyValue& InPropertyValue);
+
+	FPropertyValue GetDefaultValueForProperty(const FPropertyDesc& InProperty);
+
+	void HandlePropertyChangedEvent(TSharedPtr<const FCaptureEvent> InEvent,
+									TMap<FString, TSharedPtr<FPropertyValue>> InPropertyValueHolder);
+
+	TMap<FString, FPropertyWidgetCreator> PropertyCustomizations;
+	TMap<FString, FCommandWidgetCreator> CommandCustomizations;
+	TMap<FString, FCaptureSourceWidgetCreator> CreationParamsCustomizations;
+	TMap<FString, FCapabilityWidgetCreator> CapabilityCustomizations;
+};
\ No newline at end of file
diff --git a/Engine/Plugins/VirtualProduction/CaptureManager/Source/CaptureSourceFramework/Public/CaptureSourceError.h b/Engine/Plugins/VirtualProduction/CaptureManager/Source/CaptureSourceFramework/Public/CaptureSourceError.h
new file mode 100644
--- /dev/null
+++ b/Engine/Plugins/VirtualProduction/CaptureManager/Source/CaptureSourceFramework/Public/CaptureSourceError.h
@@ -0,0 +1,35 @@
+// Copyright Epic Games, Inc. All Rights Reserved.
+
+#pragma once
+
+#include "Containers/UnrealString.h"
+#include "Templates/ValueOrError.h"
+
+enum class ECaptureSourceError
+{
+	Ok = 0,
+	InvalidArgument,
+	NotSupported,
+	NotFound,
+	InternalError,
+};
+
+class CAPTURESOURCEFRAMEWORK_API FCaptureSourceError
+{
+public:
+
+	FCaptureSourceError(ECaptureSourceError InError, FString InMessage = "");
+
+	ECaptureSourceError GetErrorCode() const;
+	const FString& GetMessage() const;
+
+private:
+
+	ECaptureSourceError ErrorCode;
+	FString Message;
+};
+
+using FCaptureVoidResult = TValueOrError<void, FCaptureSourceError>;
+
+template<typename Value>
+using FCaptureValueResult = TValueOrError<Value, FCaptureSourceError>;
\ No newline at end of file
diff --git a/Engine/Plugins/VirtualProduction/CaptureManager/Source/CaptureSourceFramework/Public/CaptureSourceFactory.h b/Engine/Plugins/VirtualProduction/CaptureManager/Source/CaptureSourceFramework/Public/CaptureSourceFactory.h
new file mode 100644
--- /dev/null
+++ b/Engine/Plugins/VirtualProduction/CaptureManager/Source/CaptureSourceFramework/Public/CaptureSourceFactory.h
@@ -0,0 +1,61 @@
+// Copyright Epic Games, Inc. All Rights Reserved.
+
+#pragma once
+
+#include "CaptureSource.h"
+#include "CaptureSourceError.h" 
+
+class CAPTURESOURCEFRAMEWORK_API FCaptureSourceDescriptor
+{
+public:
+
+	FCaptureSourceDescriptor(const FString& InId, TArray<FPropertyDesc> InCreationParams);
+
+	void SetDiscoverable(bool bInIsDiscoverable);
+	bool IsDiscoverable() const;
+
+	void AddCreationParam(FPropertyDesc InParam);
+
+	template<typename ... FArgs>
+	void EmplaceCreationParam(FArgs&&... InArgs)
+	{
+		CreationParams.Emplace(Forward<FArgs>(InArgs)...);
+	}
+
+	const TArray<FPropertyDesc>& GetCreationParams() const;
+
+	const FString& GetId() const;
+
+private:
+
+	FString Id;
+	TArray<FPropertyDesc> CreationParams;
+	bool bIsDiscoverable = true;
+};
+
+class CAPTURESOURCEFRAMEWORK_API FCaptureSourceFactory
+{
+public:
+
+	using FCaptureSourceResult = FCaptureValueResult<TUniquePtr<FCaptureSource>>;
+
+	DECLARE_DELEGATE_OneParam(FDiscoverCallback, TArray<FCaptureSourceResult>);
+
+	virtual ~FCaptureSourceFactory() = default;
+
+	virtual FString GetCaptureSourceFactoryName() const = 0;
+	virtual FString GetCaptureSourceFactoryId() const = 0;
+
+	virtual bool IsDiscoverable() const = 0;
+	FCaptureVoidResult Discover(FDiscoverCallback InCallback);
+
+	FCaptureSourceResult CreateCaptureSource(const FString& InName, TMap<FString, FPropertyValue> InCreationParamsValue);
+	FCaptureSourceDescriptor CreateCaptureSourceDescriptor() const;
+
+private:
+
+	virtual void DiscoverImpl(FDiscoverCallback InCallback, FCaptureVoidResult& OutResult) {}
+
+	virtual FCaptureSourceResult CreateCaptureSourceImpl(const FString& InName, const TMap<FString, FPropertyValue>& InCreationParamsValue) = 0;
+	virtual void CreateCaptureSourceDescriptorImpl(FCaptureSourceDescriptor& OutDescriptor) const = 0;
+};
\ No newline at end of file
diff --git a/Engine/Plugins/VirtualProduction/CaptureManager/Source/CaptureSourceFramework/Public/CaptureSourceFrameworkModule.h b/Engine/Plugins/VirtualProduction/CaptureManager/Source/CaptureSourceFramework/Public/CaptureSourceFrameworkModule.h
new file mode 100644
--- /dev/null
+++ b/Engine/Plugins/VirtualProduction/CaptureManager/Source/CaptureSourceFramework/Public/CaptureSourceFrameworkModule.h
@@ -0,0 +1,25 @@
+// Copyright Epic Games, Inc. All Rights Reserved.
+
+#pragma once
+
+#include "Modules/ModuleManager.h"
+
+#include "CaptureSourceManager.h"
+#include "CaptureSourceDetailsCustomization.h"
+
+class CAPTURESOURCEFRAMEWORK_API FCaptureSourceFrameworkModule
+	: public IModuleInterface
+{
+public:
+
+	virtual void StartupModule() override;
+	virtual void ShutdownModule() override;
+
+	FCaptureSourceManager& GetCaptureSourceManager();
+	FCaptureSourceDetailsCustomization& GetCustomization();
+
+private:
+
+	TUniquePtr<FCaptureSourceManager> CaptureManager;
+	TUniquePtr<FCaptureSourceDetailsCustomization> Customization;
+};
\ No newline at end of file
diff --git a/Engine/Plugins/VirtualProduction/CaptureManager/Source/CaptureSourceFramework/Public/CaptureSourceHandle.h b/Engine/Plugins/VirtualProduction/CaptureManager/Source/CaptureSourceFramework/Public/CaptureSourceHandle.h
new file mode 100644
--- /dev/null
+++ b/Engine/Plugins/VirtualProduction/CaptureManager/Source/CaptureSourceFramework/Public/CaptureSourceHandle.h
@@ -0,0 +1,40 @@
+// Copyright Epic Games, Inc. All Rights Reserved.
+
+#pragma once
+
+#include "CaptureSource.h"
+
+class CAPTURESOURCEFRAMEWORK_API FCaptureSourceHandle final :
+	public UE::CaptureManager::Private::ICaptureSource,
+	public ICaptureEventSource
+{
+public:
+
+	FCaptureSourceHandle(const FCaptureSourceHandle& InOther) = delete;
+	FCaptureSourceHandle(FCaptureSourceHandle&& InOther) = delete;
+	FCaptureSourceHandle& operator=(const FCaptureSourceHandle& InOther) = delete;
+	FCaptureSourceHandle& operator=(FCaptureSourceHandle&& InOther) = delete;
+
+	virtual FCaptureVoidResult Start() final;
+	virtual FCaptureVoidResult Stop() final;
+
+	virtual FCaptureSourceCapability* GetCapability(FString InName) const final;
+	virtual TArray<FString> GetCapabilities() const final;
+
+	virtual const FString& GetCaptureSourceFactoryId() const final;
+	virtual const FString& GetName() const final;
+
+	virtual TArray<FString> GetAvailableEvents() const final;
+	virtual void SubscribeToEvent(const FString& InEventName, FCaptureEventHandler InHandler) final;
+	virtual void UnsubscribeAll() final;
+
+	virtual const TMap<FString, FPropertyValue>& GetCreationParams() const final;
+
+private:
+
+	FCaptureSource* CaptureSource;
+
+	FCaptureSourceHandle(FCaptureSource* InCaptureSource);
+
+	friend class FCaptureSourceManager;
+};
\ No newline at end of file
diff --git a/Engine/Plugins/VirtualProduction/CaptureManager/Source/CaptureSourceFramework/Public/CaptureSourceManager.h b/Engine/Plugins/VirtualProduction/CaptureManager/Source/CaptureSourceFramework/Public/CaptureSourceManager.h
new file mode 100644
--- /dev/null
+++ b/Engine/Plugins/VirtualProduction/CaptureManager/Source/CaptureSourceFramework/Public/CaptureSourceManager.h
@@ -0,0 +1,47 @@
+// Copyright Epic Games, Inc. All Rights Reserved.
+
+#pragma once
+
+#include "CaptureSourceHandle.h"
+#include "CaptureSourceFactory.h"
+
+#include "Async/CallbackSynchronizer.h"
+
+using FCaptureSourceId = int32;
+
+class CAPTURESOURCEFRAMEWORK_API FCaptureSourceManager
+{
+public:
+
+	using FCaptureSourceResult = FCaptureValueResult<FCaptureSourceId>;
+
+	DECLARE_DELEGATE_OneParam(FDiscoverCallback, TArray<FCaptureSourceResult>);
+	
+	FCaptureSourceManager();
+
+	FCaptureSourceHandle GetCaptureSource(FCaptureSourceId InCaptureSourceId);
+	TArray<FCaptureSourceId> GetCaptureSourceIds() const;
+
+	void RemoveCaptureSource(FCaptureSourceId InCaptureSourceId);
+
+	void DiscoverCaptureSources(FDiscoverCallback InCallback);
+	void DiscoverCaptureSourcesForFactory(FString InCaptureSourceFactoryId, FDiscoverCallback InCallback);
+
+	void RegisterCaptureSourceFactory(TUniquePtr<FCaptureSourceFactory> InCaptureSourceFactory);
+
+	FCaptureSourceDescriptor GetCaptureSourceFactory(FString InCaptureSourceFactoryId);
+	TArray<FCaptureSourceDescriptor> GetCaptureSourceFactoryList();
+	FCaptureSourceResult CreateCaptureSource(FString InCaptureSourceFactoryId, FString InName, TMap<FString, FPropertyValue> InCreationParamsValue);
+
+	FString GetCaptureSourceFactoryById(FCaptureSourceId InCaptureSourceId);
+
+private:
+
+	FCaptureSourceManager(const FCaptureSourceManager& InOther) = delete;
+	FCaptureSourceManager& operator=(const FCaptureSourceManager& InOther) = delete;
+
+	TMap<FCaptureSourceId, TUniquePtr<FCaptureSource>> CaptureSources;
+	TMap<FString, TUniquePtr<FCaptureSourceFactory>> CaptureSourceFactoryList;
+
+	std::atomic<FCaptureSourceId> CurrentCaptureSourceId;
+};
\ No newline at end of file
diff --git a/Engine/Plugins/VirtualProduction/CaptureManager/Source/CaptureUtils/CaptureUtils.Build.cs b/Engine/Plugins/VirtualProduction/CaptureManager/Source/CaptureUtils/CaptureUtils.Build.cs
new file mode 100644
--- /dev/null
+++ b/Engine/Plugins/VirtualProduction/CaptureManager/Source/CaptureUtils/CaptureUtils.Build.cs
@@ -0,0 +1,16 @@
+// Copyright Epic Games, Inc. All Rights Reserved.
+
+using UnrealBuildTool;
+
+public class CaptureUtils : ModuleRules
+{
+	public CaptureUtils(ReadOnlyTargetRules Target) : base(Target)
+	{
+		PCHUsage = PCHUsageMode.UseExplicitOrSharedPCHs;
+
+		PublicDependencyModuleNames.AddRange(new string[]
+		{
+			"Core",
+		});
+	}
+}
\ No newline at end of file
diff --git a/Engine/Plugins/VirtualProduction/CaptureManager/Source/CaptureUtils/Private/Async/CallbackSynchronizer.cpp b/Engine/Plugins/VirtualProduction/CaptureManager/Source/CaptureUtils/Private/Async/CallbackSynchronizer.cpp
new file mode 100644
--- /dev/null
+++ b/Engine/Plugins/VirtualProduction/CaptureManager/Source/CaptureUtils/Private/Async/CallbackSynchronizer.cpp
@@ -0,0 +1,33 @@
+// Copyright Epic Games, Inc. All Rights Reserved.
+
+#include "Async/CallbackSynchronizer.h"
+
+FCallbackSynchronizer::FCallbackSynchronizer() = default;
+
+FCallbackSynchronizer::~FCallbackSynchronizer()
+{
+	check(Counter == 0);
+}
+
+void FCallbackSynchronizer::AfterAll(FAfterAllDelegate InAfterAllDelegate, bool bExecuteIfCounterZero)
+{
+	AfterAllDelegate = MoveTemp(InAfterAllDelegate);
+
+	if (Counter == 0 && bExecuteIfCounterZero)
+	{
+		AfterAllDelegate.ExecuteIfBound();
+	}
+}
+
+TSharedPtr<FCallbackSynchronizer> FCallbackSynchronizer::Create()
+{
+	return MakeShared<FCallbackSynchronizer>();
+}
+
+void FCallbackSynchronizer::Decrease()
+{
+	if (--Counter == 0)
+	{
+		AfterAllDelegate.ExecuteIfBound();
+	}
+}
\ No newline at end of file
diff --git a/Engine/Plugins/VirtualProduction/CaptureManager/Source/CaptureUtils/Private/Async/Event.cpp b/Engine/Plugins/VirtualProduction/CaptureManager/Source/CaptureUtils/Private/Async/Event.cpp
new file mode 100644
--- /dev/null
+++ b/Engine/Plugins/VirtualProduction/CaptureManager/Source/CaptureUtils/Private/Async/Event.cpp
@@ -0,0 +1,15 @@
+// Copyright Epic Games, Inc. All Rights Reserved.
+
+#include "Async/Event.h"
+#include "Misc/ScopeRWLock.h"
+
+FCaptureEvent::FCaptureEvent(FString InName) : Name(MoveTemp(InName))
+{
+}
+
+const FString& FCaptureEvent::GetName() const
+{
+	return Name;
+}
+
+FCaptureEvent::~FCaptureEvent() = default;
\ No newline at end of file
diff --git a/Engine/Plugins/VirtualProduction/CaptureManager/Source/CaptureUtils/Private/Async/EventSourceUtils.cpp b/Engine/Plugins/VirtualProduction/CaptureManager/Source/CaptureUtils/Private/Async/EventSourceUtils.cpp
new file mode 100644
--- /dev/null
+++ b/Engine/Plugins/VirtualProduction/CaptureManager/Source/CaptureUtils/Private/Async/EventSourceUtils.cpp
@@ -0,0 +1,51 @@
+// Copyright Epic Games, Inc. All Rights Reserved.
+
+#pragma once
+
+#include "Async/EventSourceUtils.h"
+#include "Misc/ScopeRWLock.h"
+
+namespace UE::CaptureManager::Private
+{
+	FCaptureEventSourceBase::~FCaptureEventSourceBase() = default;
+
+	TArray<FString> FCaptureEventSourceBase::GetAvailableEvents() const
+	{
+		TArray<FString> Keys;
+		FRWScopeLock Guard(HandlersLock, FRWScopeLockType::SLT_ReadOnly);
+		EventToHandlersMap.GetKeys(Keys);
+		return Keys;
+	}
+
+	void FCaptureEventSourceBase::SubscribeToEvent(const FString& InEventName, FCaptureEventHandler InHandler)
+	{
+		FRWScopeLock Guard(HandlersLock, FRWScopeLockType::SLT_Write);
+		FHandlers* Handlers = EventToHandlersMap.Find(InEventName);
+		check(Handlers != nullptr);
+		Handlers->AddLambda(MoveTemp(InHandler));
+	}
+
+	void FCaptureEventSourceBase::UnsubscribeAll()
+	{
+		FRWScopeLock Guard(HandlersLock, FRWScopeLockType::SLT_Write);
+		for (TPair<FString, FHandlers>& EventHandlersPair : EventToHandlersMap)
+		{
+			EventHandlersPair.Value.Clear();
+		}
+	}
+
+	void FCaptureEventSourceBase::RegisterEvent(const FString& InEventName)
+	{
+		FRWScopeLock Guard(HandlersLock, FRWScopeLockType::SLT_Write);
+		check(EventToHandlersMap.Find(InEventName) == nullptr);
+		EventToHandlersMap.Add({ InEventName, FHandlers{} });
+	}
+
+	void FCaptureEventSourceBase::PublishEventInternal(TSharedPtr<const FCaptureEvent> InEvent) const
+	{
+		FRWScopeLock Guard(HandlersLock, FRWScopeLockType::SLT_ReadOnly);
+		const FHandlers* Handlers = EventToHandlersMap.Find(InEvent->GetName());
+		check(Handlers != nullptr);
+		Handlers->Broadcast(InEvent);
+	}
+}
\ No newline at end of file
diff --git a/Engine/Plugins/VirtualProduction/CaptureManager/Source/CaptureUtils/Private/Async/EventSourceUtils.spec.cpp b/Engine/Plugins/VirtualProduction/CaptureManager/Source/CaptureUtils/Private/Async/EventSourceUtils.spec.cpp
new file mode 100644
--- /dev/null
+++ b/Engine/Plugins/VirtualProduction/CaptureManager/Source/CaptureUtils/Private/Async/EventSourceUtils.spec.cpp
@@ -0,0 +1,136 @@
+// Copyright Epic Games, Inc. All Rights Reserved.
+
+#include "Async/Event.h"
+#include "Async/EventSourceUtils.h"
+
+#include "Misc/AutomationTest.h"
+
+#include "CoreGlobals.h"
+
+#if WITH_DEV_AUTOMATION_TESTS
+
+BEGIN_DEFINE_SPEC(FCaptureEventSourceTest, "System.Engine.Animation.LiveLink.CaptureUtils.CaptureSourceEvent", EAutomationTestFlags::EditorContext | EAutomationTestFlags::ProductFilter | EAutomationTestFlags::MediumPriority)
+
+struct FTestCaptureEvent : public FCaptureEvent
+{
+	FTestCaptureEvent(const FString& InName, int InValue) : FCaptureEvent(InName), Value(InValue)
+	{
+	}
+
+	int Value;
+};
+
+struct FTestCaptureEventSource : virtual public FCaptureEventSource
+{
+	FTestCaptureEventSource(const TArray<FString>& InEventsToRegister)
+	{
+		for (const FString& EventName : InEventsToRegister)
+		{
+			RegisterEvent(EventName);
+		}
+	}
+
+	void DoPublish(const FString& InName, int InValue)
+	{
+		PublishEvent<FTestCaptureEvent>(InName, InValue);
+	}
+};
+
+END_DEFINE_SPEC(FCaptureEventSourceTest)
+
+void FCaptureEventSourceTest::Define()
+{
+	Describe("GetAvailableEvents()", [this]()
+		{
+			It("should return the list of registered events", [this]()
+				{
+					TArray<FString> EventsToRegister = {"TestEvent1", "TestEvent2", "TestEvent3"};
+					FTestCaptureEventSource EventSource{ EventsToRegister };
+					TArray<FString> AvailableEvents = EventSource.GetAvailableEvents();
+					AvailableEvents.Sort();
+					UTEST_EQUAL("List of available events matches the expected list of event names", AvailableEvents, EventsToRegister);
+
+					return true;
+				});
+		});
+
+	Describe("PublishEvent()", [this]()
+		{
+			It("should publish event to handlers registered to the event and not publish to others", [this]()
+				{
+					constexpr int ExpectedEventValue = 5;
+					TArray<FString> EventsToRegister = { "TestEvent1", "TestEvent2"};
+					FTestCaptureEventSource EventSource{ EventsToRegister };
+
+					bool IsExpectedCalled = false;
+					bool DidChecksPass = false;
+					FCaptureEventHandler ExpectedHandler([&](TSharedPtr<const FCaptureEvent> InEvent) {
+						DidChecksPass = (InEvent->GetName() == EventsToRegister[0]) &&
+							(static_cast<const FTestCaptureEvent*>(InEvent.Get())->Value == ExpectedEventValue);
+						IsExpectedCalled = true;
+						// This needs to run on `InternalThread` because the tests are being run on the game thread so posting
+						// to it and blocking it will just deadlock it.
+					}, EDelegateExecutionThread::InternalThread);
+
+					bool IsUnexpectedCalled = false;
+					FCaptureEventHandler UnexpectedHandler([&](TSharedPtr<const FCaptureEvent> InEvent) {
+						IsUnexpectedCalled = true;
+					}, EDelegateExecutionThread::InternalThread);
+
+					EventSource.SubscribeToEvent(EventsToRegister[0], MoveTemp(ExpectedHandler));
+					EventSource.SubscribeToEvent(EventsToRegister[1], MoveTemp(UnexpectedHandler));
+
+					EventSource.DoPublish(EventsToRegister[0], ExpectedEventValue);
+
+					UTEST_TRUE("Registered handler was invoked", IsExpectedCalled);
+					UTEST_TRUE("Received event should be the published one and should have the expected value", DidChecksPass);
+					UTEST_FALSE("Handler registered to other event was not invoked", IsUnexpectedCalled);
+
+					return true;
+				});
+
+			It("should publish event to all registered handlers with all handlers sharing the same event instance", [this]()
+				{
+					constexpr int ExpectedEventValue = 6;
+					TArray<FString> EventsToRegister = { "TestEvent1"};
+					FTestCaptureEventSource EventSource{ EventsToRegister };
+
+					bool IsHandler1Called = false;
+					bool DidChecksPass1 = false;
+					const FCaptureEvent* EventPtr1 = nullptr;
+					FCaptureEventHandler Handler1([&](TSharedPtr<const FCaptureEvent> InEvent) {
+						DidChecksPass1 = (InEvent->GetName() == EventsToRegister[0]) &&
+							(static_cast<const FTestCaptureEvent*>(InEvent.Get())->Value == ExpectedEventValue);
+						IsHandler1Called = true;
+						EventPtr1 = InEvent.Get();
+						// This needs to run on `InternalThread` because the tests are being run on the game thread so posting
+						// to it and blocking it will just deadlock it.
+						}, EDelegateExecutionThread::InternalThread);
+
+					bool IsHandler2Called = false;
+					bool DidChecksPass2 = false;
+					const FCaptureEvent* EventPtr2 = nullptr;
+					FCaptureEventHandler Handler2([&](TSharedPtr<const FCaptureEvent> InEvent) {
+						DidChecksPass2 = (InEvent->GetName() == EventsToRegister[0]) &&
+							(static_cast<const FTestCaptureEvent*>(InEvent.Get())->Value == ExpectedEventValue);
+						IsHandler2Called = true;
+						EventPtr2 = InEvent.Get();
+						// This needs to run on `InternalThread` because the tests are being run on the game thread so posting
+						// to it and blocking it will just deadlock it.
+						}, EDelegateExecutionThread::InternalThread);
+
+					EventSource.SubscribeToEvent(EventsToRegister[0], MoveTemp(Handler1));
+					EventSource.SubscribeToEvent(EventsToRegister[0], MoveTemp(Handler2));
+
+					EventSource.DoPublish(EventsToRegister[0], ExpectedEventValue);
+
+					UTEST_TRUE("Registered handlers were invoked", IsHandler1Called && (IsHandler1Called == IsHandler2Called));
+					UTEST_TRUE("Received event should be the published one and should have the expected value", DidChecksPass1 && (DidChecksPass1 == DidChecksPass2));
+					UTEST_TRUE("The published event was shared amongst handlers", EventPtr1 && (EventPtr1 == EventPtr2));
+
+					return true;
+				});
+		});
+}
+
+#endif
\ No newline at end of file
diff --git a/Engine/Plugins/VirtualProduction/CaptureManager/Source/CaptureUtils/Private/Async/StopToken.cpp b/Engine/Plugins/VirtualProduction/CaptureManager/Source/CaptureUtils/Private/Async/StopToken.cpp
new file mode 100644
--- /dev/null
+++ b/Engine/Plugins/VirtualProduction/CaptureManager/Source/CaptureUtils/Private/Async/StopToken.cpp
@@ -0,0 +1,20 @@
+// Copyright Epic Games, Inc. All Rights Reserved.
+
+#include "Async/StopToken.h"
+
+FStopToken::FStopToken() :
+	SharedState(MakeShared<FSharedState>())
+{
+}
+
+void FStopToken::RequestStop()
+{
+	check(SharedState.IsValid());
+	SharedState->State.store(true);
+}
+
+bool FStopToken::IsStopRequested() const
+{
+	check(SharedState.IsValid());
+	return SharedState->State.load();
+}
\ No newline at end of file
diff --git a/Engine/Plugins/VirtualProduction/CaptureManager/Source/CaptureUtils/Private/CaptureUtilsModule.cpp b/Engine/Plugins/VirtualProduction/CaptureManager/Source/CaptureUtils/Private/CaptureUtilsModule.cpp
new file mode 100644
--- /dev/null
+++ b/Engine/Plugins/VirtualProduction/CaptureManager/Source/CaptureUtils/Private/CaptureUtilsModule.cpp
@@ -0,0 +1,9 @@
+// Copyright Epic Games, Inc. All Rights Reserved.
+
+#include "Modules/ModuleManager.h"
+
+class FCaptureUtilsModule : public IModuleInterface
+{
+};
+
+IMPLEMENT_MODULE(FCaptureUtilsModule, CaptureUtils)
diff --git a/Engine/Plugins/VirtualProduction/CaptureManager/Source/CaptureUtils/Public/Async/CallbackSynchronizer.h b/Engine/Plugins/VirtualProduction/CaptureManager/Source/CaptureUtils/Public/Async/CallbackSynchronizer.h
new file mode 100644
--- /dev/null
+++ b/Engine/Plugins/VirtualProduction/CaptureManager/Source/CaptureUtils/Public/Async/CallbackSynchronizer.h
@@ -0,0 +1,92 @@
+// Copyright Epic Games, Inc. All Rights Reserved.
+
+#pragma once
+
+#include "Templates/SharedPointer.h"
+#include "Templates/UnrealTemplate.h"
+#include "Delegates/Delegate.h"
+
+#include <atomic>
+
+template<typename ...>
+struct ArgsBundle
+{
+};
+
+template<typename T>
+struct FFunctionTraits : public FFunctionTraits<decltype(&T::operator())>
+{
+};
+
+template<typename Return, typename Object, typename ... Args>
+struct FFunctionTraits<Return(Object::*)(Args ...)>
+{
+	using ReturnType = Return;
+	using FnArgs = ArgsBundle<Args ...>;
+	using FnType = Return(Args ...);
+};
+
+template<typename Return, typename Object, typename ... Args>
+struct FFunctionTraits<Return(Object::*)(Args ...) const>
+{
+	using ReturnType = Return;
+	using FnArgs = ArgsBundle<Args ...>;
+	using FnType = Return(Args ...);
+};
+
+template<typename Return, typename ... Args>
+struct FFunctionTraits<Return(*)(Args ...)>
+{
+	using ReturnType = Return;
+	using FnArgs = ArgsBundle<Args ...>;
+	using FnType = Return(Args ...);
+};
+
+class CAPTUREUTILS_API FCallbackSynchronizer : public TSharedFromThis<FCallbackSynchronizer>
+{
+public:
+	DECLARE_DELEGATE(FAfterAllDelegate);
+
+	FCallbackSynchronizer();
+	~FCallbackSynchronizer();
+
+	template<typename FCallback>
+	auto CreateCallback(FCallback&& InCallback)
+	{
+		using FIdentityFnArgs = typename FFunctionTraits<FCallback>::FnArgs;
+		return FMakeCallbackInternal<FIdentityFnArgs>::Get(AsShared(), MoveTemp(InCallback));
+	}
+
+	void AfterAll(FAfterAllDelegate InAfterAllDelegate, bool bExecuteIfCounterZero = true);
+
+	static TSharedPtr<FCallbackSynchronizer> Create();
+
+private:
+
+	template<typename Args>
+	struct FMakeCallbackInternal;
+
+	template<typename ... Args>
+	struct FMakeCallbackInternal<ArgsBundle<Args ...> >
+	{
+		template<typename LT>
+		static auto Get(TSharedPtr<FCallbackSynchronizer> InCallbackSynchronizer, LT&& InUserFn)
+		{
+			InCallbackSynchronizer->Counter++;
+
+			return [CallbackSynchronizer = MoveTemp(InCallbackSynchronizer),
+				UserFn(MoveTemp(InUserFn))](Args&& ... InArgs) mutable
+			{
+				UserFn(Forward<Args>(InArgs) ...);
+
+				CallbackSynchronizer->Decrease();
+			};
+		}
+	};
+
+	void Decrease();
+
+	FAfterAllDelegate AfterAllDelegate;
+
+	std::atomic_int Counter;
+};
\ No newline at end of file
diff --git a/Engine/Plugins/VirtualProduction/CaptureManager/Source/CaptureUtils/Public/Async/Event.h b/Engine/Plugins/VirtualProduction/CaptureManager/Source/CaptureUtils/Public/Async/Event.h
new file mode 100644
--- /dev/null
+++ b/Engine/Plugins/VirtualProduction/CaptureManager/Source/CaptureUtils/Public/Async/Event.h
@@ -0,0 +1,55 @@
+// Copyright Epic Games, Inc. All Rights Reserved.
+
+#pragma once
+
+#include "Async/ManagedDelegate.h"
+
+#include "Containers/UnrealString.h"
+#include "Containers/Array.h"
+#include "Containers/Map.h"
+#include "Templates/SharedPointer.h"
+#include "Delegates/Delegate.h"
+#include "HAL/CriticalSection.h"
+#include "Containers/UnrealString.h"
+#include "Misc/DateTime.h"
+#include "Misc/Timespan.h"
+
+#include <atomic>
+
+// Base class for all event sublasses
+class CAPTUREUTILS_API FCaptureEvent
+{
+public:
+	const FString& GetName() const;
+	virtual ~FCaptureEvent();
+
+protected:
+	FCaptureEvent(FString InName);
+
+
+private:
+	FString Name;
+};
+
+// Macro for easy definition of events that don't carry any data
+#define CAPTURE_DEFINE_EMPTY_EVENT(ClassName, EventName) \
+	struct CAPTUREUTILS_API ClassName : public FCaptureEvent \
+	{ \
+		inline static const FString Name = TEXT(EventName); \
+		ClassName() : FCaptureEvent(Name) \
+		{ \
+		} \
+	};
+
+// The SharedPtr points to a `const` event object because the shared event might end up on multiple
+// threads in which case we'd have a thread safety issue without the event being `const`
+using FCaptureEventHandler = TManagedDelegate<TSharedPtr<const FCaptureEvent>>;
+
+// Interface for classes that wish to provide capture event subscription to their clients
+class ICaptureEventSource
+{
+public:
+	virtual TArray<FString> GetAvailableEvents() const = 0;
+	virtual void SubscribeToEvent(const FString& InEventName, FCaptureEventHandler InHandler) = 0;
+	virtual void UnsubscribeAll() = 0;
+};
\ No newline at end of file
diff --git a/Engine/Plugins/VirtualProduction/CaptureManager/Source/CaptureUtils/Public/Async/EventSourceUtils.h b/Engine/Plugins/VirtualProduction/CaptureManager/Source/CaptureUtils/Public/Async/EventSourceUtils.h
new file mode 100644
--- /dev/null
+++ b/Engine/Plugins/VirtualProduction/CaptureManager/Source/CaptureUtils/Public/Async/EventSourceUtils.h
@@ -0,0 +1,109 @@
+// Copyright Epic Games, Inc. All Rights Reserved.
+
+#pragma once
+
+#include "Async/Event.h"
+
+namespace UE::CaptureManager::Private
+{
+	class CAPTUREUTILS_API FCaptureEventSourceBase : public ICaptureEventSource
+	{
+	public:
+		virtual ~FCaptureEventSourceBase();
+
+	public:
+		virtual TArray<FString> GetAvailableEvents() const override;
+		// NOTE: SubscribeToEvent() must not be called from an event handler that is being executed on the same thread
+		virtual void SubscribeToEvent(const FString& InEventName, FCaptureEventHandler InHandler) override;
+		virtual void UnsubscribeAll() override;
+
+	protected:
+		void RegisterEvent(const FString& InEventName);
+		void PublishEventInternal(TSharedPtr<const FCaptureEvent> InEvent) const;
+
+	private:
+		using FHandlers = TMulticastDelegate<void(TSharedPtr<const FCaptureEvent>)>;
+
+		mutable FRWLock HandlersLock;
+		TMap<FString, FHandlers> EventToHandlersMap;
+
+	};
+}
+
+// Class to be publicly inherited to get basic event source functionality. All functions are thread-safe.
+class CAPTUREUTILS_API FCaptureEventSource : public UE::CaptureManager::Private::FCaptureEventSourceBase
+{
+protected:
+	template<typename EventType, typename ... Args>
+	void PublishEvent(Args&&... InArgs) const
+	{
+		PublishEventInternal(MakeShared<const EventType>(Forward<Args>(InArgs)...));
+	}
+
+	void PublishEventPtr(TSharedPtr<const FCaptureEvent> InEvent) const
+	{
+		PublishEventInternal(MoveTemp(InEvent));
+	}
+
+private:
+	using UE::CaptureManager::Private::FCaptureEventSourceBase::PublishEventInternal;
+};
+
+// Class to be publicly inherited to get event source functionality where rate of events published is limited. All functions are thread-safe.
+class CAPTUREUTILS_API FCaptureEventSourceWithLimiter : public UE::CaptureManager::Private::FCaptureEventSourceBase
+{
+public:
+	FCaptureEventSourceWithLimiter(int32 InThresholdMillis) : ThresholdMillis(InThresholdMillis), LastPublish(0)
+	{
+	}
+
+protected:
+	// When PublishIfThresholdReached() is called the publishing will mostly only occur if the time since last
+	// publish is greater than the threshold (although this isn't guaranteed and multiple events can still be
+	// published in rare cases). The unpublished events are simply dropped (i.e. not buffered so the client has
+	// to make sure that won't cause problems to the subscribers.
+	// The client can also force publish an event which will do the publishing no matter what the current time is and
+	// reset the time measurement. This is handy for publishing events whose dropping will break continuity for the
+	// client or for publishing the last event in a line of optional events to trigger the final update.
+	template<typename EventType, typename ... EventArgs>
+	bool PublishIfThresholdReached(bool bInForcePublish, EventArgs&& ... InEventArgs)
+	{
+		int64 LastPublishLocal = LastPublish.load();
+		int64 Now = FDateTime::Now().GetTicks();
+
+		if (bInForcePublish || IsThresholdPassed(LastPublishLocal, Now))
+		{
+			// Only update if nobody updated after we loaded the value. This way we don't overwrite newer time
+			// if another thread had a chance to do that before this one.
+			LastPublish.compare_exchange_strong(LastPublishLocal, Now);
+			PublishEventInternal(MakeShared<const EventType>(Forward<EventArgs>(InEventArgs)...));
+			return true;
+		}
+		return false;
+	}
+
+	// Always publishes the event while completely ignoring the threshold mechanism. If you want to force publishing
+	// of an event AND update the "last publish" timestamp, please use PublishIfThresholdReached(true, ...)
+	void PublishEventIgnoreThresholdPtr(TSharedPtr<const FCaptureEvent> InEvent)
+	{
+		PublishEventInternal(MoveTemp(InEvent));
+	}
+
+	template<typename EventType, typename ... EventArgs>
+	void PublishEventIgnoreThreshold(EventArgs&&... InEventArgs)
+	{
+		PublishEventPtr(MakeShared<const EventType>(Forward<EventArgs>(InEventArgs)...));
+	}
+
+private:
+	using UE::CaptureManager::Private::FCaptureEventSourceBase::PublishEventInternal;
+
+	bool IsThresholdPassed(const FDateTime& InLastPublish, const FDateTime& InNow) const
+	{
+		FTimespan TimespanElapsed = InNow - InLastPublish;
+		return TimespanElapsed.GetTotalMilliseconds() >= ThresholdMillis;
+	}
+
+	const int32 ThresholdMillis;
+	std::atomic<int64> LastPublish;
+};
\ No newline at end of file
diff --git a/Engine/Plugins/VirtualProduction/CaptureManager/Source/CaptureUtils/Public/Async/ManagedDelegate.h b/Engine/Plugins/VirtualProduction/CaptureManager/Source/CaptureUtils/Public/Async/ManagedDelegate.h
new file mode 100644
--- /dev/null
+++ b/Engine/Plugins/VirtualProduction/CaptureManager/Source/CaptureUtils/Public/Async/ManagedDelegate.h
@@ -0,0 +1,117 @@
+// Copyright Epic Games, Inc. All Rights Reserved.
+
+#pragma once
+
+#include "Delegates/Delegate.h"
+#include "Async/Async.h"
+
+enum class EDelegateExecutionThread
+{
+	GameThread = 0,
+	InternalThread // Thread on which the function invoking the delegate is being run on
+};
+
+namespace UE::CaptureManager::Private
+{
+template<typename ... Args>
+static void ExecuteDelegate(TDelegate<void(Args...)> InDelegate, EDelegateExecutionThread InThread, Args&&... InArgs)
+{
+	if (InThread == EDelegateExecutionThread::GameThread)
+	{
+		TDelegate<void()> Task = TDelegate<void()>::CreateLambda([InDelegate = (MoveTemp(InDelegate))](Args... InArgs)
+		{
+			InDelegate.Execute(Forward<Args>(InArgs)...);
+		}, Forward<Args>(InArgs)...);
+
+		AsyncTask(ENamedThreads::GameThread, [InTask = MoveTemp(Task)]()
+		{
+			InTask.Execute();
+		});
+	}
+	else
+	{
+		InDelegate.Execute(Forward<Args>(InArgs)...);
+	}
+}
+}
+
+template<typename ... Args>
+class TManagedDelegate
+{
+public:
+	using Type = TDelegate<void(Args...)>;
+
+	TManagedDelegate() = default;
+
+	template<typename Func>
+	TManagedDelegate(Func InFunc, EDelegateExecutionThread InThread = EDelegateExecutionThread::GameThread)
+		: TManagedDelegate(Type::CreateLambda(MoveTemp(InFunc)), InThread)
+	{
+	}
+
+	TManagedDelegate(Type InDelegate, EDelegateExecutionThread InThread = EDelegateExecutionThread::GameThread)
+		: Delegate(MoveTemp(InDelegate))
+		, ExecutionThread(InThread)
+	{
+	}
+
+	void operator()(Args... InArgs)
+	{
+		if (Delegate.IsBound())
+		{
+			UE::CaptureManager::Private::ExecuteDelegate(Delegate, ExecutionThread, Forward<Args>(InArgs)...);
+		}
+	}
+
+	void operator()(Args... InArgs) const
+	{
+		if (Delegate.IsBound())
+		{
+			UE::CaptureManager::Private::ExecuteDelegate(Delegate, ExecutionThread, Forward<Args>(InArgs)...);
+		}
+	}
+
+private:
+
+	Type Delegate;
+	EDelegateExecutionThread ExecutionThread;
+};
+
+template<typename ... Args>
+class TManagedMulticastDelegate
+{
+public:
+	using DelegateType = TDelegate<void(Args...)>;
+
+	TManagedMulticastDelegate() = default;
+
+	template<typename Func>
+	void Add(Func InFunc, EDelegateExecutionThread InThread = EDelegateExecutionThread::GameThread)
+	{
+		Add(DelegateType::CreateLambda(MoveTemp(InFunc)), InThread);
+	}
+
+	void Add(DelegateType InDelegate, EDelegateExecutionThread InThread = EDelegateExecutionThread::GameThread)
+	{
+		Delegate.AddLambda([InDelegate = MoveTemp(InDelegate), InThread](Args... InArgs) mutable
+		{
+			UE::CaptureManager::Private::ExecuteDelegate(InDelegate, InThread, Forward<Args>(InArgs)...);
+		});
+	}
+
+	void operator()(Args... InArgs)
+	{
+		Delegate.Broadcast(Forward<Args>(InArgs)...);
+	}
+
+	void operator()(Args... InArgs) const
+	{
+		Delegate.Broadcast(Forward<Args>(InArgs)...);
+	}
+
+private:
+
+	using Type = TMulticastDelegate<void(Args...)>;
+
+	Type Delegate;
+};
\ No newline at end of file
diff --git a/Engine/Plugins/VirtualProduction/CaptureManager/Source/CaptureUtils/Public/Async/StopToken.h b/Engine/Plugins/VirtualProduction/CaptureManager/Source/CaptureUtils/Public/Async/StopToken.h
new file mode 100644
--- /dev/null
+++ b/Engine/Plugins/VirtualProduction/CaptureManager/Source/CaptureUtils/Public/Async/StopToken.h
@@ -0,0 +1,28 @@
+// Copyright Epic Games, Inc. All Rights Reserved.
+
+#pragma once
+
+#include "Templates/SharedPointer.h"
+
+class CAPTUREUTILS_API FStopToken
+{
+public:
+	FStopToken();
+
+	FStopToken(const FStopToken& InOther) = default;
+	FStopToken(FStopToken&& InOther) = default;
+	FStopToken& operator=(const FStopToken& InOther) = default;
+	FStopToken& operator=(FStopToken&& InOther) = default;
+
+	void RequestStop();
+	bool IsStopRequested() const;
+
+private:
+
+	struct FSharedState
+	{
+		std::atomic_bool State = false;
+	};
+
+	TSharedPtr<FSharedState> SharedState;
+};
\ No newline at end of file
diff --git a/Engine/Plugins/VirtualProduction/CaptureManager/Source/CaptureUtils/Public/Async/Task.h b/Engine/Plugins/VirtualProduction/CaptureManager/Source/CaptureUtils/Public/Async/Task.h
new file mode 100644
--- /dev/null
+++ b/Engine/Plugins/VirtualProduction/CaptureManager/Source/CaptureUtils/Public/Async/Task.h
@@ -0,0 +1,74 @@
+// Copyright Epic Games, Inc. All Rights Reserved.
+
+#pragma once
+
+#include "Templates/Function.h"
+#include "Async/AsyncWork.h"
+#include "Async/StopToken.h"
+
+#include <atomic>
+
+class FCancelableAsyncTask
+{
+public:
+	using FTaskFunction = TFunction<void(const FStopToken&)>;
+
+	FCancelableAsyncTask(FTaskFunction InTaskFunction) :
+		AsyncTask(MakeUnique<FAsyncTask<FAsyncTaskInternal>>(StopToken, MoveTemp(InTaskFunction)))
+	{
+	}
+
+	~FCancelableAsyncTask()
+	{
+		Cancel();
+		AsyncTask->EnsureCompletion();
+	}
+
+	bool IsDone()
+	{
+		return AsyncTask->IsDone();
+	}
+
+	void StartSync()
+	{
+		AsyncTask->StartSynchronousTask();
+	}
+
+	void StartAsync()
+	{
+		AsyncTask->StartBackgroundTask();
+	}
+
+	void Cancel()
+	{
+		StopToken.RequestStop();
+	}
+
+private:
+	class FAsyncTaskInternal : public FNonAbandonableTask
+	{
+	public:
+		FAsyncTaskInternal(const FStopToken& InStopToken, FTaskFunction&& InTaskFunction) :
+			StopToken(InStopToken), TaskFunction(MoveTemp(InTaskFunction))
+		{
+			check(TaskFunction != nullptr);
+		}
+
+		void DoWork()
+		{
+			TaskFunction(StopToken);
+		}
+
+		FORCEINLINE TStatId GetStatId() const
+		{
+			RETURN_QUICK_DECLARE_CYCLE_STAT(FCancelableAsyncTask_FAsyncTaskInternal, STATGROUP_ThreadPoolAsyncTasks);
+		}
+
+	private:
+		const FStopToken& StopToken;
+		FTaskFunction TaskFunction;
+	};
+
+	FStopToken StopToken;
+	TUniquePtr<FAsyncTask<FAsyncTaskInternal>> AsyncTask;
+};
\ No newline at end of file
