LiveLink Hub Capture Source implementation

diff --git a/Engine/Plugins/VirtualProduction/CaptureManager/Source/CaptureSourceFramework/Private/CaptureSourceCapability.cpp b/Engine/Plugins/VirtualProduction/CaptureManager/Source/CaptureSourceFramework/Private/CaptureSourceCapability.cpp
--- a/Engine/Plugins/VirtualProduction/CaptureManager/Source/CaptureSourceFramework/Private/CaptureSourceCapability.cpp
+++ b/Engine/Plugins/VirtualProduction/CaptureManager/Source/CaptureSourceFramework/Private/CaptureSourceCapability.cpp
@@ -9,7 +9,10 @@ FCaptureSourceCapabilityDesc::FCaptureSourceCapabilityDesc(FString InName, TArra
 {
 }
 
-const FString& FCapturePropertyChangedEvent::EventName = TEXT("PropertyChanged");
+const FString FCapturePropertyChangedEvent::EventName = TEXT("PropertyChanged");
+const FString FCapturePropertyChangedEvent::Property = TEXT("Property");
+const FString FCapturePropertyChangedEvent::PropertyValue = TEXT("Value");
+
 FCapturePropertyChangedEvent::FCapturePropertyChangedEvent(const FString& InName, FPropertyValue InValue)
 	: FCaptureEvent(EventName)
 	, Name(InName)
@@ -20,7 +23,11 @@ FCapturePropertyChangedEvent::FCapturePropertyChangedEvent(const FString& InName
 FCaptureSourceCapability::FCaptureSourceCapability(FString InName)
 	: Name(MoveTemp(InName))
 {
-	RegisterEvent(FCapturePropertyChangedEvent::EventName);
+	AddEvent(FEventDesc(FCapturePropertyChangedEvent::EventName, 
+						{
+							FPropertyDesc(FCapturePropertyChangedEvent::Property, FPropertyDesc::EType::String),
+							FPropertyDesc(FCapturePropertyChangedEvent::PropertyValue, FPropertyDesc::EType::Any)
+						}));
 }
 
 void FCaptureSourceCapability::PublishPropertyChangedEvent(const FString& InName, FPropertyValue InValue)
@@ -71,6 +78,11 @@ TArray<FCommandDesc> FCaptureSourceCapability::GetCommands() const
 	return Commands.GetCommands();
 }
 
+TArray<FEventDesc> FCaptureSourceCapability::GetEvents() const
+{
+	return Events.GetEvents();
+}
+
 FString FCaptureSourceCapability::GetName() const
 {
 	return Name;
@@ -112,6 +124,13 @@ void FCaptureSourceCapability::AddCommand(FCommandDesc InCommand)
 	Commands.AddCommand(MoveTemp(InCommand));
 }
 
+void FCaptureSourceCapability::AddEvent(FEventDesc InEvent)
+{
+	RegisterEvent(InEvent.Name);
+
+	Events.AddEvent(MoveTemp(InEvent));
+}
+
 void FCaptureSourceCapability::ExecuteCommand(TSharedPtr<FCommandBase> InCommand)
 {
 	const FString& CommandName = InCommand->GetName();
diff --git a/Engine/Plugins/VirtualProduction/CaptureManager/Source/CaptureSourceFramework/Private/CaptureSourceCapabilityEvent.cpp b/Engine/Plugins/VirtualProduction/CaptureManager/Source/CaptureSourceFramework/Private/CaptureSourceCapabilityEvent.cpp
new file mode 100644
--- /dev/null
+++ b/Engine/Plugins/VirtualProduction/CaptureManager/Source/CaptureSourceFramework/Private/CaptureSourceCapabilityEvent.cpp
@@ -0,0 +1,27 @@
+// Copyright Epic Games, Inc. All Rights Reserved.
+
+#include "CaptureSourceCapabilityEvent.h"
+
+FEventDesc::FEventDesc(FString InName, TArray<FPropertyDesc> InParams)
+	: Name(MoveTemp(InName))
+	, Params(MoveTemp(InParams))
+{
+}
+
+void FEventList::AddEvent(FEventDesc InEventDesc)
+{
+	Events.Add(MoveTemp(InEventDesc));
+}
+
+TArray<FEventDesc> FEventList::GetEvents() const
+{
+	return Events;
+}
+
+void FEventList::CheckEventExists(const FString& InName) const
+{
+	check(Events.ContainsByPredicate([&InName](const FEventDesc& InEvent)
+	{
+		return InEvent.Name == InName;
+	}));
+}
\ No newline at end of file
diff --git a/Engine/Plugins/VirtualProduction/CaptureManager/Source/CaptureSourceFramework/Private/CaptureSourceCapabilityProperty.cpp b/Engine/Plugins/VirtualProduction/CaptureManager/Source/CaptureSourceFramework/Private/CaptureSourceCapabilityProperty.cpp
--- a/Engine/Plugins/VirtualProduction/CaptureManager/Source/CaptureSourceFramework/Private/CaptureSourceCapabilityProperty.cpp
+++ b/Engine/Plugins/VirtualProduction/CaptureManager/Source/CaptureSourceFramework/Private/CaptureSourceCapabilityProperty.cpp
@@ -78,6 +78,8 @@ bool FPropertyList::CheckPropertyValueType(FPropertyDesc::EType InType, const FP
 			return InValue.IsType<TArray<FString>>();
 		case FPropertyDesc::EType::ArrayFloatingPoint:
 			return InValue.IsType<TArray<double>>();
+		case FPropertyDesc::EType::Any:
+			return true;
 		default:
 			return false;
 	}
diff --git a/Engine/Plugins/VirtualProduction/CaptureManager/Source/CaptureSourceFramework/Private/CaptureSourceDetailsCustomization.cpp b/Engine/Plugins/VirtualProduction/CaptureManager/Source/CaptureSourceFramework/Private/CaptureSourceDetailsCustomization.cpp
--- a/Engine/Plugins/VirtualProduction/CaptureManager/Source/CaptureSourceFramework/Private/CaptureSourceDetailsCustomization.cpp
+++ b/Engine/Plugins/VirtualProduction/CaptureManager/Source/CaptureSourceFramework/Private/CaptureSourceDetailsCustomization.cpp
@@ -403,15 +403,10 @@ TSharedPtr<SWidget> FCaptureSourceDetailsCustomization::MakeValueWidget(FCapture
 			});
 			break;
 		case FPropertyDesc::EType::ArrayBool:
-			Widget = SNullWidget::NullWidget;
-			break;
 		case FPropertyDesc::EType::ArrayNumber:
-			Widget = SNullWidget::NullWidget;
-			break;
 		case FPropertyDesc::EType::ArrayString:
-			Widget = SNullWidget::NullWidget;
-			break;
 		case FPropertyDesc::EType::ArrayFloatingPoint:
+		case FPropertyDesc::EType::Any:
 		default:
 			Widget = SNullWidget::NullWidget;
 	}
@@ -473,15 +468,10 @@ TSharedPtr<SWidget> FCaptureSourceDetailsCustomization::MakeValueWidget(const FP
 			});
 			break;
 		case FPropertyDesc::EType::ArrayBool:
-			Widget = SNullWidget::NullWidget;
-			break;
 		case FPropertyDesc::EType::ArrayNumber:
-			Widget = SNullWidget::NullWidget;
-			break;
 		case FPropertyDesc::EType::ArrayString:
-			Widget = SNullWidget::NullWidget;
-			break;
 		case FPropertyDesc::EType::ArrayFloatingPoint:
+		case FPropertyDesc::EType::Any:
 		default:
 			Widget = SNullWidget::NullWidget;
 	}
@@ -504,12 +494,10 @@ FPropertyValue FCaptureSourceDetailsCustomization::GetDefaultValueForProperty(co
 		case FPropertyDesc::EType::String:
 			return MakePropertyValue<FString>(TEXT(""));
 		case FPropertyDesc::EType::ArrayBool:
-			return MakePropertyValue(FEmptyVariantState());
 		case FPropertyDesc::EType::ArrayNumber:
-			return MakePropertyValue(FEmptyVariantState());
 		case FPropertyDesc::EType::ArrayString:
-			return MakePropertyValue(FEmptyVariantState());
 		case FPropertyDesc::EType::ArrayFloatingPoint:
+		case FPropertyDesc::EType::Any:
 		default:
 			return MakePropertyValue(FEmptyVariantState());
 	}
diff --git a/Engine/Plugins/VirtualProduction/CaptureManager/Source/CaptureSourceFramework/Public/CaptureSourceCapability.h b/Engine/Plugins/VirtualProduction/CaptureManager/Source/CaptureSourceFramework/Public/CaptureSourceCapability.h
--- a/Engine/Plugins/VirtualProduction/CaptureManager/Source/CaptureSourceFramework/Public/CaptureSourceCapability.h
+++ b/Engine/Plugins/VirtualProduction/CaptureManager/Source/CaptureSourceFramework/Public/CaptureSourceCapability.h
@@ -6,6 +6,7 @@
 
 #include "CaptureSourceCapabilityProperty.h"
 #include "CaptureSourceCapabilityCommand.h"
+#include "CaptureSourceCapabilityEvent.h"
 
 struct FCaptureSourceCapabilityDesc
 {
@@ -19,7 +20,9 @@ struct FCaptureSourceCapabilityDesc
 struct CAPTURESOURCEFRAMEWORK_API FCapturePropertyChangedEvent : public FCaptureEvent
 {
 public:
-	static const FString& EventName;
+	static const FString EventName;
+	static const FString Property;
+	static const FString PropertyValue;
 
 	FCapturePropertyChangedEvent(const FString& InName, FPropertyValue InValue);
 
@@ -35,7 +38,7 @@ public:
 	DECLARE_DELEGATE_RetVal_OneParam(FPropertyValue, FPropertyGetter, const FString&);
 	DECLARE_DELEGATE_OneParam(FCommandHandler, TSharedPtr<FCommandBase>);
 
-	virtual ~FCaptureSourceCapability() = default;
+	virtual ~FCaptureSourceCapability() override = default;
 
 	void PublishPropertyChangedEvent(const FString& InName, FPropertyValue InValue);
 
@@ -46,6 +49,7 @@ public:
 
 	TArray<FPropertyDesc> GetProperties() const;
 	TArray<FCommandDesc> GetCommands() const;
+	TArray<FEventDesc> GetEvents() const;
 
 	FString GetName() const;
 
@@ -61,11 +65,13 @@ protected:
 
 	void AddProperty(FPropertyDesc InPropertyDesc);
 	void AddCommand(FCommandDesc InCommand);
+	void AddEvent(FEventDesc InEvent);
 
 private:
 
 	FPropertyList Properties;
 	FCommandList Commands;
+	FEventList Events;
 
 	FString Name;
 
diff --git a/Engine/Plugins/VirtualProduction/CaptureManager/Source/CaptureSourceFramework/Public/CaptureSourceCapabilityEvent.h b/Engine/Plugins/VirtualProduction/CaptureManager/Source/CaptureSourceFramework/Public/CaptureSourceCapabilityEvent.h
new file mode 100644
--- /dev/null
+++ b/Engine/Plugins/VirtualProduction/CaptureManager/Source/CaptureSourceFramework/Public/CaptureSourceCapabilityEvent.h
@@ -0,0 +1,31 @@
+// Copyright Epic Games, Inc. All Rights Reserved.
+
+#pragma once
+
+#include "Async/EventSourceUtils.h"
+
+#include "CaptureSourceCapabilityProperty.h"
+
+struct CAPTURESOURCEFRAMEWORK_API FEventDesc
+{
+	FEventDesc(FString InName, TArray<FPropertyDesc> InParams);
+
+	FString Name;
+	TArray<FPropertyDesc> Params;
+};
+
+class FEventList
+{
+public:
+
+	FEventList() = default;
+
+	void AddEvent(FEventDesc InEventDesc);
+	TArray<FEventDesc> GetEvents() const;
+
+	void CheckEventExists(const FString& InName) const;
+
+private:
+
+	TArray<FEventDesc> Events;
+};
diff --git a/Engine/Plugins/VirtualProduction/CaptureManager/Source/CaptureSourceFramework/Public/CaptureSourceCapabilityProperty.h b/Engine/Plugins/VirtualProduction/CaptureManager/Source/CaptureSourceFramework/Public/CaptureSourceCapabilityProperty.h
--- a/Engine/Plugins/VirtualProduction/CaptureManager/Source/CaptureSourceFramework/Public/CaptureSourceCapabilityProperty.h
+++ b/Engine/Plugins/VirtualProduction/CaptureManager/Source/CaptureSourceFramework/Public/CaptureSourceCapabilityProperty.h
@@ -20,7 +20,8 @@ struct CAPTURESOURCEFRAMEWORK_API FPropertyDesc
 		ArrayBool,
 		ArrayNumber,
 		ArrayString,
-		ArrayFloatingPoint
+		ArrayFloatingPoint,
+		Any
 	};
 
 	enum class EAccess
diff --git a/Engine/Plugins/VirtualProduction/CaptureManager/Source/LiveLinkHubCaptureMessaging/LiveLinkHubCaptureMessaging.Build.cs b/Engine/Plugins/VirtualProduction/CaptureManager/Source/LiveLinkHubCaptureMessaging/LiveLinkHubCaptureMessaging.Build.cs
new file mode 100644
--- /dev/null
+++ b/Engine/Plugins/VirtualProduction/CaptureManager/Source/LiveLinkHubCaptureMessaging/LiveLinkHubCaptureMessaging.Build.cs
@@ -0,0 +1,27 @@
+// Copyright Epic Games, Inc. All Rights Reserved.
+
+using UnrealBuildTool;
+using System.IO;
+
+public class LiveLinkHubCaptureMessaging : ModuleRules
+{
+	public LiveLinkHubCaptureMessaging(ReadOnlyTargetRules Target) : base(Target)
+	{
+		PCHUsage = PCHUsageMode.UseExplicitOrSharedPCHs;
+
+		PublicDependencyModuleNames.AddRange(new string[]
+		{
+			"Core",
+			"CoreUObject",
+			"Engine",
+			"CaptureUtils"
+		});
+
+		PrivateDependencyModuleNames.AddRange(new string[]
+		{
+			"UnrealEd",
+			"CaptureSourceFramework",
+			"MessagingCommon",
+		});
+	}
+}
diff --git a/Engine/Plugins/VirtualProduction/CaptureManager/Source/LiveLinkHubCaptureMessaging/Public/LiveLinkHubCaptureMessages.h b/Engine/Plugins/VirtualProduction/CaptureManager/Source/LiveLinkHubCaptureMessaging/Public/LiveLinkHubCaptureMessages.h
new file mode 100644
--- /dev/null
+++ b/Engine/Plugins/VirtualProduction/CaptureManager/Source/LiveLinkHubCaptureMessaging/Public/LiveLinkHubCaptureMessages.h
@@ -0,0 +1,265 @@
+// Copyright Epic Games, Inc. All Rights Reserved.
+
+#pragma once
+
+#include "UObject/Class.h"
+
+#include "Containers/UnrealString.h"
+
+#include "Misc/Guid.h"
+
+#include "LiveLinkHubCaptureMessages.generated.h"
+
+namespace UE::LiveLinkHubCaptureMessages
+{
+	constexpr uint16 Version = 1;
+}
+
+UENUM()
+enum class EStatus
+{
+	Ok = 0,
+	InvalidArgument = 1
+};
+
+USTRUCT()
+struct FBaseMessage
+{
+	GENERATED_BODY()
+
+	UPROPERTY()
+	uint16 Version = UE::LiveLinkHubCaptureMessages::Version;
+
+	UPROPERTY()
+	FGuid Guid;
+};
+
+USTRUCT()
+struct FBaseResponse : public FBaseMessage
+{
+	GENERATED_BODY()
+
+	UPROPERTY()
+	EStatus Status = EStatus::Ok;
+
+	UPROPERTY()
+	FString Message;
+
+	UPROPERTY()
+	FGuid RequestGuid;
+};
+
+USTRUCT()
+struct FConnectRequest : public FBaseMessage
+{
+	GENERATED_BODY()
+};
+
+USTRUCT()
+struct FConnectResponse : public FBaseResponse
+{
+	GENERATED_BODY()
+};
+
+USTRUCT()
+struct FDiscoveryRequest
+{
+	GENERATED_BODY()
+};
+
+USTRUCT()
+struct FDiscoveryResponse
+{
+	GENERATED_BODY()
+
+	UPROPERTY()
+	FString Name;
+};
+
+USTRUCT()
+struct FPingRequest : public FBaseMessage
+{
+	GENERATED_BODY();
+};
+
+USTRUCT()
+struct FPingResponse : public FBaseResponse
+{
+	GENERATED_BODY();
+};
+
+USTRUCT()
+struct FGetCapabilitiesRequest : public FBaseMessage
+{
+	GENERATED_BODY()
+};
+
+USTRUCT()
+struct FGetCapabilitiesResponse : public FBaseResponse
+{
+	GENERATED_BODY()
+
+	UPROPERTY()
+	TArray<FString> Capabilities;
+};
+
+UENUM()
+enum class EValueType
+{
+	Unknown = -1,
+	Bool = 0,
+	Integer = 1,
+	FloatingPoint = 2,
+	String = 3
+};
+
+UENUM()
+enum class EValueAccess
+{
+	ReadWrite = 0,
+	ReadOnly = 1,
+};
+
+USTRUCT()
+struct FValueDescriptor
+{
+	GENERATED_BODY()
+
+	UPROPERTY()
+	EValueType Type = EValueType::Bool;
+
+	UPROPERTY()
+	EValueAccess Access = EValueAccess::ReadOnly;
+};
+
+USTRUCT()
+struct FMapValuesDescriptor
+{
+	GENERATED_BODY()
+
+	UPROPERTY()
+	TMap<FString, FValueDescriptor> Values;
+};
+
+USTRUCT()
+struct FGetCapabilityRequest : public FBaseMessage
+{
+	GENERATED_BODY()
+
+	UPROPERTY()
+	FString Capability;
+};
+
+USTRUCT()
+struct FGetCapabilityResponse : public FBaseResponse
+{
+	GENERATED_BODY()
+
+	UPROPERTY()
+	FMapValuesDescriptor Properties;
+
+	UPROPERTY()
+	TMap<FString, FMapValuesDescriptor> Commands;
+
+	UPROPERTY()
+	TMap<FString, FMapValuesDescriptor> Events;
+};
+
+USTRUCT()
+struct FGetPropertyValueRequest : public FBaseMessage
+{
+	GENERATED_BODY()
+
+	UPROPERTY()
+	FString Capability;
+
+	UPROPERTY()
+	FString Property;
+};
+
+USTRUCT()
+struct FValue
+{
+	GENERATED_BODY()
+
+	UPROPERTY()
+	EValueType Type = EValueType::Bool;
+
+	UPROPERTY(meta = (EditCondition = "Type == EValueType::Bool"))
+	bool BoolValue = false;
+
+	UPROPERTY(meta = (EditCondition = "Type == EValueType::Integer"))
+	int64 IntegerValue = 0;
+
+	UPROPERTY(meta = (EditCondition = "Type == EValueType::FloatingPoint"))
+	double FloatingPointValue = 0.f;
+
+	UPROPERTY(meta = (EditCondition = "Type == EValueType::String"))
+	FString StringValue;
+};
+
+USTRUCT()
+struct FGetPropertyValueResponse : public FBaseResponse
+{
+	GENERATED_BODY()
+
+	UPROPERTY()
+	FValue Value;
+};
+
+USTRUCT()
+struct FSetPropertyValueRequest : public FBaseMessage
+{
+	GENERATED_BODY()
+
+	UPROPERTY()
+	FString Capability;
+	
+	UPROPERTY()
+	FString Property;
+
+	UPROPERTY()
+	FValue Value;
+};
+
+USTRUCT()
+struct FSetPropertyValueResponse : public FBaseResponse
+{
+	GENERATED_BODY()
+};
+
+USTRUCT()
+struct FExecuteCommandRequest : public FBaseMessage
+{
+	GENERATED_BODY()
+
+	UPROPERTY()
+	FString Capability;
+
+	UPROPERTY()
+	FString Command;
+
+	UPROPERTY()
+	TMap<FString, FValue> Params;
+};
+
+USTRUCT()
+struct FExecuteCommandResponse : public FBaseResponse
+{
+	GENERATED_BODY()
+};
+
+USTRUCT()
+struct FEventMessage : public FBaseMessage
+{
+	GENERATED_BODY()
+
+	UPROPERTY()
+	FString Capability;
+
+	UPROPERTY()
+	FString Name;
+
+	UPROPERTY()
+	TMap<FString, FValue> Params;
+};
diff --git a/Engine/Plugins/VirtualProduction/CaptureManager/Source/LiveLinkHubCaptureMessaging/Public/LiveLinkHubCaptureMessaggingModule.cpp b/Engine/Plugins/VirtualProduction/CaptureManager/Source/LiveLinkHubCaptureMessaging/Public/LiveLinkHubCaptureMessaggingModule.cpp
new file mode 100644
--- /dev/null
+++ b/Engine/Plugins/VirtualProduction/CaptureManager/Source/LiveLinkHubCaptureMessaging/Public/LiveLinkHubCaptureMessaggingModule.cpp
@@ -0,0 +1,10 @@
+// Copyright Epic Games, Inc. All Rights Reserved.
+
+#include "Modules/ModuleManager.h"
+
+class FLiveLinkHubCaptureMessagingModule
+	: public IModuleInterface
+{
+};
+
+IMPLEMENT_MODULE(FLiveLinkHubCaptureMessagingModule, LiveLinkHubCaptureMessaging)
\ No newline at end of file
diff --git a/Engine/Plugins/VirtualProduction/CaptureManager/Source/LiveLinkHubCaptureSource/LiveLinkHubCaptureSource.Build.cs b/Engine/Plugins/VirtualProduction/CaptureManager/Source/LiveLinkHubCaptureSource/LiveLinkHubCaptureSource.Build.cs
new file mode 100644
--- /dev/null
+++ b/Engine/Plugins/VirtualProduction/CaptureManager/Source/LiveLinkHubCaptureSource/LiveLinkHubCaptureSource.Build.cs
@@ -0,0 +1,36 @@
+// Copyright Epic Games, Inc. All Rights Reserved.
+
+using UnrealBuildTool;
+using System.IO;
+
+public class LiveLinkHubCaptureSource : ModuleRules
+{
+	public LiveLinkHubCaptureSource(ReadOnlyTargetRules Target) : base(Target)
+	{
+		PCHUsage = PCHUsageMode.UseExplicitOrSharedPCHs;
+
+		PublicDependencyModuleNames.AddRange(new string[]
+		{
+			"Core",
+			"CoreUObject",
+			"Engine",
+			"CaptureUtils"
+		});
+
+		PrivateDependencyModuleNames.AddRange(new string[]
+		{
+			"InputCore",
+			"Slate",
+			"Json",
+			"SlateCore",
+			"AudioEditor",
+			"UnrealEd",
+			"RenderCore",
+			"CaptureSourceFramework",
+			"LiveLinkHubCaptureMessaging",
+			"ContentBrowser",
+			"DesktopPlatform",
+			"MessagingCommon"
+		});
+	}
+}
diff --git a/Engine/Plugins/VirtualProduction/CaptureManager/Source/LiveLinkHubCaptureSource/Private/Capabilities/MessagingCapability.cpp b/Engine/Plugins/VirtualProduction/CaptureManager/Source/LiveLinkHubCaptureSource/Private/Capabilities/MessagingCapability.cpp
new file mode 100644
--- /dev/null
+++ b/Engine/Plugins/VirtualProduction/CaptureManager/Source/LiveLinkHubCaptureSource/Private/Capabilities/MessagingCapability.cpp
@@ -0,0 +1,275 @@
+// Copyright Epic Games, Inc. All Rights Reserved.
+
+#include "Capabilities/MessagingCapability.h"
+
+const FString FMessagingEvent::Name = TEXT("MessagingEvent");
+
+FMessagingEvent::FMessagingEvent(const FString& InEventName, const FString& InCapability)
+	: FCaptureEvent(Name)
+	, EventName(InEventName)
+	, Capability(InCapability)
+{
+}
+
+const FString FMessagingCapability::Name = TEXT("Messaging");
+const float FMessagingCapability::KeepAliveInterval = 5.0f; // seconds
+const float FMessagingCapability::KeepAliveTimeout = 3.0f; // seconds
+
+FMessagingCapability::FMessagingCapability()
+	: FCaptureSourceCapability(Name)
+	, Builder(*Name)
+	, bIsConnected(false)
+{
+	Builder.Handling<FEventMessage>(this, &FMessagingCapability::EventHandler)
+		.Handling<FConnectResponse>(this, &FMessagingCapability::ConnectHandler)
+		.Handling<FPingResponse>(this, &FMessagingCapability::ResponseHandler<FPingResponse>)
+		.Handling<FGetCapabilitiesResponse>(this, &FMessagingCapability::ResponseHandler<FGetCapabilitiesResponse>)
+		.Handling<FGetCapabilityResponse>(this, &FMessagingCapability::ResponseHandler<FGetCapabilityResponse>)
+		.Handling<FGetPropertyValueResponse>(this, &FMessagingCapability::ResponseHandler<FGetPropertyValueResponse>)
+		.Handling<FSetPropertyValueResponse>(this, &FMessagingCapability::ResponseHandler<FSetPropertyValueResponse>)
+		.Handling<FExecuteCommandResponse>(this, &FMessagingCapability::ResponseHandler<FExecuteCommandResponse>);
+
+	RegisterEvent(FMessagingEvent::Name);
+}
+
+void FMessagingCapability::Initialize()
+{
+	Endpoint = Builder.Build();
+
+	Endpoint->Subscribe<FEventMessage>();
+}
+
+FMessageEndpointBuilder& FMessagingCapability::GetBuilder()
+{
+	return Builder;
+}
+
+void FMessagingCapability::SetAddress(FMessageAddress InAddress)
+{
+	ClientAddress = MoveTemp(InAddress);
+}
+
+void FMessagingCapability::SetDisconnectHandler(FDisconnectHandler InHandler)
+{
+	Handler = MoveTemp(InHandler);
+}
+
+void FMessagingCapability::Connect(FMessageCallback InCallback)
+{
+	if (bIsConnected.load())
+	{
+		return;
+	}
+
+	check(ClientAddress.IsValid());
+
+	FConnectRequest* Request = FMessageEndpoint::MakeMessage<FConnectRequest>();
+	Request->Guid = FGuid::NewGuid();
+
+	AddContext(Request->Guid, MoveTemp(InCallback));
+	Endpoint->Send(Request, ClientAddress);
+}
+
+void FMessagingCapability::SendGetCapabilities(FMessageCallback InCallback)
+{
+	if (!bIsConnected.load())
+	{
+		return;
+	}
+
+	FGetCapabilitiesRequest* Request = FMessageEndpoint::MakeMessage<FGetCapabilitiesRequest>();
+	SendRequest(Request, MoveTemp(InCallback));
+}
+
+void FMessagingCapability::SendGetCapability(FString InCapability, FMessageCallback InCallback)
+{
+	if (!bIsConnected.load())
+	{
+		return;
+	}
+
+	FGetCapabilityRequest* Request = FMessageEndpoint::MakeMessage<FGetCapabilityRequest>();
+	Request->Capability = MoveTemp(InCapability);
+
+	SendRequest(Request, MoveTemp(InCallback));
+}
+
+void FMessagingCapability::SendGetPropertyValue(FString InCapability, FString InProperty, FMessageCallback InCallback)
+{
+	if (!bIsConnected.load())
+	{
+		return;
+	}
+
+	FGetPropertyValueRequest* Request = FMessageEndpoint::MakeMessage<FGetPropertyValueRequest>();
+	Request->Capability = MoveTemp(InCapability);
+	Request->Property = MoveTemp(InProperty);
+
+	SendRequest(Request, MoveTemp(InCallback));
+}
+
+void FMessagingCapability::SendSetPropertyValue(FString InCapability, FString InProperty, FPropertyValue InValue, FMessageCallback InCallback)
+{
+	if (!bIsConnected.load())
+	{
+		return;
+	}
+
+	FSetPropertyValueRequest* Request = FMessageEndpoint::MakeMessage<FSetPropertyValueRequest>();
+	Request->Capability = MoveTemp(InCapability);
+	Request->Property = MoveTemp(InProperty);
+
+	Request->Value = ParsePropertyValue(MoveTemp(InValue));
+
+	SendRequest(Request, MoveTemp(InCallback));
+}
+
+void FMessagingCapability::SendExecuteCommand(FString InCapability, TSharedPtr<FCommandBase> InCommand, FMessageCallback InCallback)
+{
+	if (!bIsConnected.load())
+	{
+		return;
+	}
+
+	FExecuteCommandRequest* Request = FMessageEndpoint::MakeMessage<FExecuteCommandRequest>();
+	Request->Capability = MoveTemp(InCapability);
+	Request->Command = InCommand->GetName();
+
+	for (const TPair<FString, FPropertyValue>& Value : InCommand->GetParamValues())
+	{
+		Request->Params.Add(Value.Key, ParsePropertyValue(Value.Value));
+	}
+
+	SendRequest(Request, MoveTemp(InCallback));
+}
+
+void FMessagingCapability::ConnectHandler(const FConnectResponse& InResponse, const TSharedRef<IMessageContext, ESPMode::ThreadSafe>& InContext)
+{
+	check(ClientAddress == InContext->GetSender());
+
+	if (InResponse.Status == EStatus::Ok)
+	{
+		bIsConnected.store(true);
+
+		KeepAlive = FTSTicker::GetCoreTicker().AddTicker(FTickerDelegate::CreateRaw(this, &FMessagingCapability::OnKeepAliveInterval), KeepAliveInterval);
+	}
+
+	ResponseHandler(InResponse, InContext);
+}
+
+void FMessagingCapability::AddContext(FGuid InGuid, FMessageCallback InCallback)
+{
+	FScopeLock Lock(&Mutex);
+	Contexts.Add(MoveTemp(InGuid), MoveTemp(InCallback));
+}
+
+void FMessagingCapability::ResponseHandler(const FName& InResponseName, const FBaseResponse& InResponse)
+{
+	FScopeLock Lock(&Mutex);
+	if (Contexts.Contains(InResponse.RequestGuid))
+	{
+		FMessageCallback Callback;
+		Contexts.RemoveAndCopyValue(InResponse.RequestGuid, Callback);
+
+		Callback.ExecuteIfBound(InResponseName, InResponse);
+	}
+}
+
+void FMessagingCapability::EventHandler(const FEventMessage& InEvent, const TSharedRef<IMessageContext, ESPMode::ThreadSafe>&)
+{
+	if (!bIsConnected.load())
+	{
+		return;
+	}
+
+	TSharedPtr<FMessagingEvent> Event = MakeShared<FMessagingEvent>(InEvent.Name, InEvent.Capability);
+
+	for (const TPair<FString, FValue>& Param : InEvent.Params)
+	{
+		Event->Params.Add(Param.Key, ParseValue(Param.Value));
+	}
+
+	PublishEventPtr(MoveTemp(Event));
+}
+
+bool FMessagingCapability::OnKeepAliveInterval(float InDeltaTime)
+{
+	AsyncTask(ENamedThreads::Type::AnyThread, [this]()
+	{
+		if (!bIsConnected.load())
+		{
+			return;
+		}
+
+		TPromise<void> Promise;
+		TFuture<void> Future = Promise.GetFuture();
+
+		FPingRequest* Request = FMessageEndpoint::MakeMessage<FPingRequest>();
+
+		SendRequest(Request, FMessageCallback::CreateLambda([&Promise](const FName& InName, const FBaseResponse& InResponse) mutable
+		{
+			check(InName == FPingResponse::StaticStruct()->GetFName());
+
+			Promise.SetValue();
+		}));
+
+		bool Result = Future.WaitFor(FTimespan::FromSeconds(KeepAliveTimeout));
+
+		if (!Result)
+		{
+			Future.Reset();
+			Promise.SetValue(); // Setting a value to avoid error when destructing the promise
+
+			Handler.ExecuteIfBound();
+
+			bIsConnected.store(false);
+			FTSTicker::GetCoreTicker().RemoveTicker(KeepAlive);
+		}
+	});
+
+	return true;
+}
+
+FValue FMessagingCapability::ParsePropertyValue(FPropertyValue InValue)
+{
+	FValue Value;
+	if (InValue.IsType<bool>())
+	{
+		Value.BoolValue = InValue.Get<bool>();
+		Value.Type = EValueType::Bool;
+	}
+	else if (InValue.IsType<int64>())
+	{
+		Value.IntegerValue = InValue.Get<int64>();
+		Value.Type = EValueType::Integer;
+	}
+	else if (InValue.IsType<double>())
+	{
+		Value.FloatingPointValue = InValue.Get<double>();
+		Value.Type = EValueType::FloatingPoint;
+	}
+	else if (InValue.IsType<FString>())
+	{
+		Value.StringValue = InValue.Get<FString>();
+		Value.Type = EValueType::String;
+	}
+
+	return Value;
+}
+
+FPropertyValue FMessagingCapability::ParseValue(const FValue& InValue)
+{
+	switch (InValue.Type)
+	{
+		case EValueType::Bool:
+			return MakePropertyValue(InValue.BoolValue);
+		case EValueType::Integer:
+			return MakePropertyValue(InValue.IntegerValue);
+		case EValueType::FloatingPoint:
+			return MakePropertyValue(InValue.FloatingPointValue);
+		case EValueType::String:
+			return MakePropertyValue(InValue.StringValue);
+		default:
+			return MakePropertyValue(FEmptyVariantState());
+	}
+}
\ No newline at end of file
diff --git a/Engine/Plugins/VirtualProduction/CaptureManager/Source/LiveLinkHubCaptureSource/Private/LiveLinkHubCaptureSource.cpp b/Engine/Plugins/VirtualProduction/CaptureManager/Source/LiveLinkHubCaptureSource/Private/LiveLinkHubCaptureSource.cpp
new file mode 100644
--- /dev/null
+++ b/Engine/Plugins/VirtualProduction/CaptureManager/Source/LiveLinkHubCaptureSource/Private/LiveLinkHubCaptureSource.cpp
@@ -0,0 +1,214 @@
+// Copyright Epic Games, Inc. All Rights Reserved.
+
+#include "LiveLinkHubCaptureSource.h"
+
+#include "LiveLinkHubCaptureSourceFactory.h"
+#include "LiveLinkHubCaptureSourceLog.h"
+
+#include "MessageEndpointBuilder.h"
+
+#include "CaptureSourceFrameworkModule.h"
+#include "CaptureSourceDetailsCustomization.h"
+
+#include "Capabilities/DiskCapacityCapability.h"
+#include "Capabilities/RecordingCapability.h"
+
+#include "Widgets/SNullWidget.h"
+
+DEFINE_LOG_CATEGORY(LogLiveLinkHubCaptureSource);
+
+FLiveLinkHubCaptureSource::FLiveLinkHubCaptureSource(const FString& InName)
+	: FCaptureSource(FLiveLinkHubCaptureSourceFactory::Id, InName)
+{
+	using FCustomization = FCaptureSourceDetailsCustomization;
+	using FCreator = FCustomization::FCapabilityWidgetCreator;
+
+	FCaptureSourceDetailsCustomization& Customization =
+		FModuleManager::LoadModuleChecked<FCaptureSourceFrameworkModule>("CaptureSourceFramework").GetCustomization();
+
+	Customization.
+		RegisterCapabilityCustomization(FLiveLinkHubCaptureSourceFactory::Id, 
+										FMessagingCapability::Name, FCreator::CreateRaw(this, &FLiveLinkHubCaptureSource::CreateMessagingWidget));
+
+	TUniquePtr<FMessagingCapability> Messaging = MakeUnique<FMessagingCapability>();
+
+	Messaging->SubscribeToEvent(FMessagingEvent::Name, FCaptureEventHandler(FCaptureEventHandler::Type::CreateRaw(this, &FLiveLinkHubCaptureSource::EventHandler), EDelegateExecutionThread::InternalThread));
+
+	MessagingCapability = Messaging.Get();
+	OwnedCapabilities.Add(FMessagingCapability::Name, MoveTemp(Messaging));
+
+	AddCapability(MessagingCapability);
+
+	MessagingCapability->Initialize();
+
+	MessagingCapability->SetDisconnectHandler(FMessagingCapability::FDisconnectHandler::CreateRaw(this, &FLiveLinkHubCaptureSource::OnDisconnect));
+}
+
+FCaptureVoidResult FLiveLinkHubCaptureSource::Start()
+{
+	MessagingCapability->Connect(FMessagingCapability::FMessageCallback::CreateLambda([this](const FName& InResponseName, const FBaseResponse& InResponse)
+	{
+		check(InResponseName == FConnectResponse::StaticStruct()->GetFName());
+
+		if (InResponse.Status == EStatus::Ok)
+		{
+			MessagingCapability->SendGetCapabilities(FMessagingCapability::FMessageCallback::CreateRaw(this, &FLiveLinkHubCaptureSource::HandleGetCapabilities));
+
+			PublishEvent<FReachableEvent>(true);
+		}
+		else
+		{
+			PublishEvent<FReachableEvent>(false);
+		}
+	}));
+
+	return MakeValue();
+}
+
+FCaptureVoidResult FLiveLinkHubCaptureSource::Stop()
+{
+	return MakeValue();
+}
+
+void FLiveLinkHubCaptureSource::HandleGetCapabilities(const FName& InResponseName, const FBaseResponse& InResponse)
+{
+	check(InResponseName == FGetCapabilitiesResponse::StaticStruct()->GetFName());
+
+	if (InResponse.Status == EStatus::Ok)
+	{
+		const FGetCapabilitiesResponse& Response = static_cast<const FGetCapabilitiesResponse&>(InResponse);
+
+		for (const FString& Capability : Response.Capabilities)
+		{
+			if (TUniquePtr<FCaptureSourceCapability> CapabilityPtr = CreateCapability(Capability); CapabilityPtr.IsValid())
+			{
+				OwnedCapabilities.Add(Capability, MoveTemp(CapabilityPtr));
+			}
+		}
+	}
+	else
+	{
+		UE_LOG(LogLiveLinkHubCaptureSource, Error, TEXT("Failed to get the capabilities from LiveLink Hub: %s"), *InResponse.Message);
+	}
+}
+
+TUniquePtr<FCaptureSourceCapability> FLiveLinkHubCaptureSource::CreateCapability(const FString& InName)
+{
+	if (InName == FDiskCapacityCapability::Name)
+	{
+		TUniquePtr<FDiskCapacityCapability> DiskCapacity = CreateCapability<FDiskCapacityCapability>();
+
+		DiskCapacity->SetRemainingSpacePropertyHandler(FCaptureSourceCapability::FPropertyGetter::CreateRaw(this, &FLiveLinkHubCaptureSource::PropertyGetter, FDiskCapacityCapability::Name));
+		DiskCapacity->SetTotalSpacePropertyHandler(FCaptureSourceCapability::FPropertyGetter::CreateRaw(this, &FLiveLinkHubCaptureSource::PropertyGetter, FDiskCapacityCapability::Name));
+
+		DiskCapacity->SetUpdateCmdHandler(FCaptureSourceCapability::FCommandHandler::CreateRaw(this, &FLiveLinkHubCaptureSource::CommandHandler, FDiskCapacityCapability::Name));
+
+		AddCapability(DiskCapacity.Get());
+
+		// Get state
+		MessagingCapability->SendGetPropertyValue(FDiskCapacityCapability::Name, FDiskCapacityCapability::TotalSpace, 
+												  FMessagingCapability::FMessageCallback::CreateRaw(this, &FLiveLinkHubCaptureSource::GetPropertyHandler, FDiskCapacityCapability::TotalSpace));
+		MessagingCapability->SendGetPropertyValue(FDiskCapacityCapability::Name, FDiskCapacityCapability::RemainingSpace, 
+												  FMessagingCapability::FMessageCallback::CreateRaw(this, &FLiveLinkHubCaptureSource::GetPropertyHandler, FDiskCapacityCapability::RemainingSpace));
+		MessagingCapability->SendGetPropertyValue(FDiskCapacityCapability::Name, FDiskCapacityCapability::UsedSpace, 
+												  FMessagingCapability::FMessageCallback::CreateRaw(this, &FLiveLinkHubCaptureSource::GetPropertyHandler, FDiskCapacityCapability::UsedSpace));
+
+		return DiskCapacity;
+	}
+	else if (InName == FRecordingCapability::Name)
+	{
+		TUniquePtr<FRecordingCapability> Recording = CreateCapability<FRecordingCapability>();
+
+		Recording->SetStartRecordingHandler(FCaptureSourceCapability::FCommandHandler::CreateRaw(this, &FLiveLinkHubCaptureSource::CommandHandler, FRecordingCapability::Name));
+		Recording->SetStopRecordingHandler(FCaptureSourceCapability::FCommandHandler::CreateRaw(this, &FLiveLinkHubCaptureSource::CommandHandler, FRecordingCapability::Name));
+
+		AddCapability(Recording.Get());
+
+		return Recording;
+	}
+	else
+	{
+		UE_LOG(LogLiveLinkHubCaptureSource, Warning, TEXT("Unsupported capability: %s"), *InName);
+		return nullptr;
+	}
+}
+
+FPropertyValue FLiveLinkHubCaptureSource::PropertyGetter(const FString& InProperty, FString InCapability)
+{
+	FScopeLock Lock(&Mutex);
+
+	if (PropertyValues.Contains(InProperty))
+	{
+		return PropertyValues[InProperty];
+	}
+
+	return FPropertyValue();
+}
+
+void FLiveLinkHubCaptureSource::CommandHandler(TSharedPtr<FCommandBase> InCommand, FString InCapability)
+{
+	MessagingCapability->SendExecuteCommand(InCapability, 
+											MoveTemp(InCommand),
+											FMessagingCapability::FMessageCallback::CreateRaw(this, &FLiveLinkHubCaptureSource::HandleCommandResponse, InCapability));
+}
+
+void FLiveLinkHubCaptureSource::HandleCommandResponse(const FName& InName, const FBaseResponse& InResponse, FString InCapability)
+{
+	check(InName == FExecuteCommandResponse::StaticStruct()->GetFName());
+
+	if (InResponse.Status != EStatus::Ok)
+	{
+		UE_LOG(LogLiveLinkHubCaptureSource, Error, TEXT("Failed to execute a command for [%s]: %s"), *InCapability, *InResponse.Message);
+	}
+}
+
+void FLiveLinkHubCaptureSource::GetPropertyHandler(const FName& InName, const FBaseResponse& InResponse, FString InProperty)
+{
+	check(InName == FGetPropertyValueResponse::StaticStruct()->GetFName());
+
+	const FGetPropertyValueResponse& Response = static_cast<const FGetPropertyValueResponse&>(InResponse);
+
+	FPropertyValue Value = FMessagingCapability::ParseValue(Response.Value);
+
+	FScopeLock Lock(&Mutex);
+
+	if (!PropertyValues.Contains(InProperty))
+	{
+		PropertyValues.Add(MoveTemp(InProperty), MoveTemp(Value));
+	}
+	else
+	{
+		PropertyValues[MoveTemp(InProperty)] = MoveTemp(Value);
+	}
+}
+
+void FLiveLinkHubCaptureSource::EventHandler(TSharedPtr<const FCaptureEvent> InEvent)
+{
+	check(InEvent->GetName() == FMessagingEvent::Name);
+
+	TSharedPtr<const FMessagingEvent> Event = StaticCastSharedPtr<const FMessagingEvent>(InEvent);
+
+	if (Event->EventName == FCapturePropertyChangedEvent::EventName)
+	{
+		checkf(OwnedCapabilities.Contains(Event->Capability), TEXT("Source doesn't contain given capability"));
+
+		OwnedCapabilities[Event->Capability]->PublishPropertyChangedEvent(
+			Event->Params[FCapturePropertyChangedEvent::Property].Get<FString>(), Event->Params[FCapturePropertyChangedEvent::PropertyValue]);
+	}
+	else
+	{
+		UE_LOG(LogLiveLinkHubCaptureSource, Warning, TEXT("Unsupported event arrived %s"), *(Event->EventName));
+	}
+}
+
+void FLiveLinkHubCaptureSource::OnDisconnect()
+{
+	UE_LOG(LogTemp, Display, TEXT("Connection to the LiveLinkHub lost"));
+
+	PublishEvent<FReachableEvent>(false);
+}
+
+TSharedPtr<SWidget> FLiveLinkHubCaptureSource::CreateMessagingWidget(int32, FCaptureSourceCapability*)
+{
+	return SNullWidget::NullWidget;
+}
\ No newline at end of file
diff --git a/Engine/Plugins/VirtualProduction/CaptureManager/Source/LiveLinkHubCaptureSource/Private/LiveLinkHubCaptureSource.h b/Engine/Plugins/VirtualProduction/CaptureManager/Source/LiveLinkHubCaptureSource/Private/LiveLinkHubCaptureSource.h
new file mode 100644
--- /dev/null
+++ b/Engine/Plugins/VirtualProduction/CaptureManager/Source/LiveLinkHubCaptureSource/Private/LiveLinkHubCaptureSource.h
@@ -0,0 +1,52 @@
+// Copyright Epic Games, Inc. All Rights Reserved.
+
+#pragma once
+
+#include "CaptureSource.h"
+
+#include "Capabilities/MessagingCapability.h"
+
+#include "MessageEndpoint.h"
+
+#include "LiveLinkHubCaptureMessages.h"
+
+#include "Widgets/SWidget.h"
+
+class FLiveLinkHubCaptureSource 
+	: public FCaptureSource
+{
+public:
+
+	FLiveLinkHubCaptureSource(const FString& InName);
+
+	virtual FCaptureVoidResult Start() override;
+	virtual FCaptureVoidResult Stop() override;
+
+private:
+
+	void HandleGetCapabilities(const FName& InResponseName, const FBaseResponse& InResponse);
+
+	TUniquePtr<FCaptureSourceCapability> CreateCapability(const FString& InName);
+
+	template<typename T, typename ... Args>
+	TUniquePtr<T> CreateCapability(Args&&... InArgs)
+	{
+		return MakeUnique<T>(Forward<Args>(InArgs)...);
+	}
+
+	FPropertyValue PropertyGetter(const FString& InProperty, FString InCapability);
+	void CommandHandler(TSharedPtr<FCommandBase> InCommand, FString InCapability);
+	void HandleCommandResponse(const FName& InName, const FBaseResponse& InResponse, FString InCapability);
+
+	void GetPropertyHandler(const FName& InName, const FBaseResponse& InResponse, FString InProperty);
+	void EventHandler(TSharedPtr<const FCaptureEvent> InEvent);
+	void OnDisconnect();
+
+	TSharedPtr<SWidget> CreateMessagingWidget(int32, FCaptureSourceCapability*);
+
+	TMap<FString, TUniquePtr<FCaptureSourceCapability>> OwnedCapabilities;
+	FMessagingCapability* MessagingCapability;
+
+	FCriticalSection Mutex;
+	TMap<FString, FPropertyValue> PropertyValues;
+};
\ No newline at end of file
diff --git a/Engine/Plugins/VirtualProduction/CaptureManager/Source/LiveLinkHubCaptureSource/Private/LiveLinkHubCaptureSourceFactory.cpp b/Engine/Plugins/VirtualProduction/CaptureManager/Source/LiveLinkHubCaptureSource/Private/LiveLinkHubCaptureSourceFactory.cpp
new file mode 100644
--- /dev/null
+++ b/Engine/Plugins/VirtualProduction/CaptureManager/Source/LiveLinkHubCaptureSource/Private/LiveLinkHubCaptureSourceFactory.cpp
@@ -0,0 +1,36 @@
+// Copyright Epic Games, Inc. All Rights Reserved.
+
+#include "LiveLinkHubCaptureSourceFactory.h"
+
+#include "LiveLinkHubCaptureSource.h"
+
+const FString FLiveLinkHubCaptureSourceFactory::Id = TEXT("LiveLinkHubCaptureSource");
+const FString FLiveLinkHubCaptureSourceFactory::Name = TEXT("LiveLink Hub Capture Source");
+
+FLiveLinkHubCaptureSourceFactory::FLiveLinkHubCaptureSourceFactory() = default;
+
+FString FLiveLinkHubCaptureSourceFactory::GetCaptureSourceFactoryName() const
+{
+	return Name;
+}
+
+FString FLiveLinkHubCaptureSourceFactory::GetCaptureSourceFactoryId() const
+{
+	return Id;
+}
+
+bool FLiveLinkHubCaptureSourceFactory::IsDiscoverable() const
+{
+	return false;
+}
+
+FLiveLinkHubCaptureSourceFactory::FCaptureSourceResult FLiveLinkHubCaptureSourceFactory::CreateCaptureSourceImpl(const FString& InName, 
+																												 const TMap<FString, FPropertyValue>& InCreationParamsValue)
+{
+	return MakeValue(MakeUnique<FLiveLinkHubCaptureSource>(InName));
+}
+
+void FLiveLinkHubCaptureSourceFactory::CreateCaptureSourceDescriptorImpl(FCaptureSourceDescriptor& OutDescriptor) const
+{
+	OutDescriptor.SetDiscoverable(IsDiscoverable());
+}
\ No newline at end of file
diff --git a/Engine/Plugins/VirtualProduction/CaptureManager/Source/LiveLinkHubCaptureSource/Private/LiveLinkHubCaptureSourceLog.h b/Engine/Plugins/VirtualProduction/CaptureManager/Source/LiveLinkHubCaptureSource/Private/LiveLinkHubCaptureSourceLog.h
new file mode 100644
--- /dev/null
+++ b/Engine/Plugins/VirtualProduction/CaptureManager/Source/LiveLinkHubCaptureSource/Private/LiveLinkHubCaptureSourceLog.h
@@ -0,0 +1,8 @@
+// Copyright Epic Games, Inc. All Rights Reserved.
+
+#pragma once
+
+#include "Logging/LogMacros.h"
+
+DECLARE_LOG_CATEGORY_EXTERN(LogLiveLinkHubCaptureSource, Log, All);
+DECLARE_LOG_CATEGORY_EXTERN(LogMessagingCapability, Log, All);
\ No newline at end of file
diff --git a/Engine/Plugins/VirtualProduction/CaptureManager/Source/LiveLinkHubCaptureSource/Private/LiveLinkHubCaptureSourceModule.cpp b/Engine/Plugins/VirtualProduction/CaptureManager/Source/LiveLinkHubCaptureSource/Private/LiveLinkHubCaptureSourceModule.cpp
new file mode 100644
--- /dev/null
+++ b/Engine/Plugins/VirtualProduction/CaptureManager/Source/LiveLinkHubCaptureSource/Private/LiveLinkHubCaptureSourceModule.cpp
@@ -0,0 +1,20 @@
+// Copyright Epic Games, Inc. All Rights Reserved.
+
+#include "LiveLinkHubCaptureSourceModule.h"
+
+#include "CaptureSourceFrameworkModule.h"
+
+#include "LiveLinkHubCaptureSourceFactory.h"
+
+void FLiveLinkHubCaptureSourceModule::StartupModule()
+{
+	FCaptureSourceFrameworkModule& Module = FModuleManager::LoadModuleChecked<FCaptureSourceFrameworkModule>("CaptureSourceFramework");
+
+	Module.GetCaptureSourceManager().RegisterCaptureSourceFactory(MakeUnique<FLiveLinkHubCaptureSourceFactory>());
+}
+
+void FLiveLinkHubCaptureSourceModule::ShutdownModule()
+{
+}
+
+IMPLEMENT_MODULE(FLiveLinkHubCaptureSourceModule, LiveLinkHubCaptureSource)
\ No newline at end of file
diff --git a/Engine/Plugins/VirtualProduction/CaptureManager/Source/LiveLinkHubCaptureSource/Private/LiveLinkHubCaptureSourceModule.h b/Engine/Plugins/VirtualProduction/CaptureManager/Source/LiveLinkHubCaptureSource/Private/LiveLinkHubCaptureSourceModule.h
new file mode 100644
--- /dev/null
+++ b/Engine/Plugins/VirtualProduction/CaptureManager/Source/LiveLinkHubCaptureSource/Private/LiveLinkHubCaptureSourceModule.h
@@ -0,0 +1,16 @@
+// Copyright Epic Games, Inc. All Rights Reserved.
+
+#pragma once
+
+#include "Modules/ModuleManager.h"
+
+class FLiveLinkHubCaptureSourceModule
+	: public IModuleInterface
+{
+public:
+
+	//~Begin IModuleInterface interface
+	virtual void StartupModule() override;
+	virtual void ShutdownModule() override;
+	//~End IModuleInterface interface
+};
\ No newline at end of file
diff --git a/Engine/Plugins/VirtualProduction/CaptureManager/Source/LiveLinkHubCaptureSource/Public/Capabilities/MessagingCapability.h b/Engine/Plugins/VirtualProduction/CaptureManager/Source/LiveLinkHubCaptureSource/Public/Capabilities/MessagingCapability.h
new file mode 100644
--- /dev/null
+++ b/Engine/Plugins/VirtualProduction/CaptureManager/Source/LiveLinkHubCaptureSource/Public/Capabilities/MessagingCapability.h
@@ -0,0 +1,116 @@
+// Copyright Epic Games, Inc. All Rights Reserved.
+
+#pragma once
+
+#include "CaptureSourceCapability.h"
+
+#include "MessageEndpointBuilder.h"
+#include "MessageEndpoint.h"
+
+#include "LiveLinkHubCaptureMessages.h"
+
+#include "HAL/CriticalSection.h"
+
+#include "Containers/Ticker.h"
+
+struct LIVELINKHUBCAPTURESOURCE_API FMessagingEvent : public FCaptureEvent
+{
+	static const FString Name;
+
+	FMessagingEvent(const FString& InEventName, const FString& InCapability);
+
+	const FString EventName;
+	const FString Capability;
+	TMap<FString, FPropertyValue> Params;
+};
+
+class LIVELINKHUBCAPTURESOURCE_API FMessagingCapability : public FCaptureSourceCapability
+{
+public:
+
+	DECLARE_DELEGATE_TwoParams(FMessageCallback, const FName&, const FBaseResponse&);
+	DECLARE_DELEGATE(FDisconnectHandler);
+
+	static const FString Name;
+	static const float KeepAliveInterval;
+	static const float KeepAliveTimeout;
+
+	FMessagingCapability();
+
+	FMessageEndpointBuilder& GetBuilder();
+
+	void Initialize();
+
+	void SendGetCapabilities(FMessageCallback InCallback);
+	void SendGetCapability(FString InCapability, FMessageCallback InCallback);
+
+	void SendGetPropertyValue(FString InCapability, FString InProperty, FMessageCallback InCallback);
+	void SendSetPropertyValue(FString InCapability, FString InProperty, FPropertyValue InValue, FMessageCallback InCallback);
+
+	void SendExecuteCommand(FString InCapability, TSharedPtr<FCommandBase> InCommand, FMessageCallback InCallback);
+
+	void SetAddress(FMessageAddress InAddress);
+	void SetDisconnectHandler(FDisconnectHandler InHandler);
+
+	void Connect(FMessageCallback InCallback);
+
+	static FValue ParsePropertyValue(FPropertyValue InValue);
+	static FPropertyValue ParseValue(const FValue& InValue);
+
+	template<typename T>
+	void SendMessage(T* InMessage, bool bIsReliable)
+	{
+		if (!bIsConnected.load())
+		{
+			return;
+		}
+
+		if (!bIsReliable)
+		{
+			Endpoint->Send(InMessage, ClientAddress);
+		}
+		else
+		{
+			Endpoint->Send(InMessage, EMessageFlags::Reliable, nullptr, { ClientAddress }, FTimespan::Zero(), FDateTime::MaxValue());
+		}
+	}
+
+private:
+
+	template<typename T>
+	void SendRequest(T* Request, FMessageCallback InCallback)
+	{
+		Request->Guid = FGuid::NewGuid();
+
+		AddContext(Request->Guid, MoveTemp(InCallback));
+		Endpoint->Send(Request, ClientAddress);
+	}
+
+	void AddContext(FGuid InGuid, FMessageCallback InCallback);
+
+	void ConnectHandler(const FConnectResponse& InResponse, const TSharedRef<IMessageContext, ESPMode::ThreadSafe>& InContext);
+
+	template<typename Response>
+	void ResponseHandler(const Response& InResponse, const TSharedRef<IMessageContext, ESPMode::ThreadSafe>&)
+	{
+		ResponseHandler(Response::StaticStruct()->GetFName(), InResponse);
+	}
+
+	void ResponseHandler(const FName& InResponseName, const FBaseResponse& InResponse);
+
+	void EventHandler(const FEventMessage& InEvent, const TSharedRef<IMessageContext, ESPMode::ThreadSafe>&);
+
+	bool OnKeepAliveInterval(float InDeltaTime);
+	
+	FMessageEndpointBuilder Builder;
+	TSharedPtr<FMessageEndpoint> Endpoint;
+	FMessageAddress ClientAddress;
+
+	FCriticalSection Mutex;
+	TMap<FGuid, FMessageCallback> Contexts;
+
+	std::atomic_bool bIsConnected;
+	FTSTicker::FDelegateHandle KeepAlive;
+
+	FDisconnectHandler Handler;
+};
\ No newline at end of file
diff --git a/Engine/Plugins/VirtualProduction/CaptureManager/Source/LiveLinkHubCaptureSource/Public/LiveLinkHubCaptureSourceFactory.h b/Engine/Plugins/VirtualProduction/CaptureManager/Source/LiveLinkHubCaptureSource/Public/LiveLinkHubCaptureSourceFactory.h
new file mode 100644
--- /dev/null
+++ b/Engine/Plugins/VirtualProduction/CaptureManager/Source/LiveLinkHubCaptureSource/Public/LiveLinkHubCaptureSourceFactory.h
@@ -0,0 +1,26 @@
+// Copyright Epic Games, Inc. All Rights Reserved.
+
+#pragma once
+
+#include "CaptureSourceFactory.h"
+
+class LIVELINKHUBCAPTURESOURCE_API FLiveLinkHubCaptureSourceFactory final : public FCaptureSourceFactory
+{
+public:
+
+	static const FString Id;
+	static const FString Name;
+
+	FLiveLinkHubCaptureSourceFactory();
+
+	virtual FString GetCaptureSourceFactoryName() const;
+	virtual FString GetCaptureSourceFactoryId() const;
+
+	virtual bool IsDiscoverable() const;
+
+private:
+
+	virtual FCaptureSourceResult CreateCaptureSourceImpl(const FString& InName, const TMap<FString, FPropertyValue>& InCreationParamsValue);
+	virtual void CreateCaptureSourceDescriptorImpl(FCaptureSourceDescriptor& OutDescriptor) const;
+
+};
\ No newline at end of file
