CaptureManager plugin file missing from the previous merge

diff --git a/Engine/Plugins/VirtualProduction/CaptureManager/CaptureManager.uplugin b/Engine/Plugins/VirtualProduction/CaptureManager/CaptureManager.uplugin
new file mode 100644
--- /dev/null
+++ b/Engine/Plugins/VirtualProduction/CaptureManager/CaptureManager.uplugin
@@ -0,0 +1,45 @@
+{
+	"FileVersion": 1,
+	"Version": 1,
+	"VersionName": "1.0.0",
+	"FriendlyName": "Capture Manager",
+	"Description": "The Capture Manager allows control and monitoring of the capture device in Unreal Engine",
+	"Category": "Virtual Production",
+	"CreatedBy": "Epic Games, Inc.",
+	"CreatedByURL": "http://epicgames.com",
+	"DocsURL": "",
+	"MarketplaceURL": "",
+	"SupportURL": "",
+	"CanContainContent": true,
+	"EnabledByDefault": false,
+	"IsBetaVersion": false,
+	"IsExperimentalVersion": true,
+	"Installed": false,
+	"SupportedPrograms": [ "LiveLinkHub" ],
+	"Modules": [
+		{
+			"Name": "CaptureSourceFramework",
+			"Type": "EditorAndProgram",
+			"ProgramAllowList": [ "LiveLinkHub" ],
+			"LoadingPhase": "PreDefault"
+		},
+		{
+			"Name": "CaptureUtils",
+			"Type": "EditorAndProgram",
+			"ProgramAllowList": [ "LiveLinkHub" ],
+			"LoadingPhase": "PreDefault"
+		},
+		{
+			"Name": "LiveLinkHubCaptureSource",
+			"Type": "Editor",
+			"LoadingPhase": "Default"
+		},
+		{
+			"Name": "LiveLinkHubCaptureMessaging",
+			"Type": "EditorAndProgram",
+			"ProgramAllowList": [ "LiveLinkHub" ],
+			"LoadingPhase": "Default"
+		}
+	],
+	"Plugins": []
+}
\ No newline at end of file
