Introducing a CPS protocol framework

diff --git a/Engine/Plugins/VirtualProduction/CaptureManager/CaptureManager.uplugin b/Engine/Plugins/VirtualProduction/CaptureManager/CaptureManager.uplugin
--- a/Engine/Plugins/VirtualProduction/CaptureManager/CaptureManager.uplugin
+++ b/Engine/Plugins/VirtualProduction/CaptureManager/CaptureManager.uplugin
@@ -39,6 +39,12 @@
 			"Type": "EditorAndProgram",
 			"ProgramAllowList": [ "LiveLinkHub" ],
 			"LoadingPhase": "Default"
+		},
+		{
+			"Name": "CaptureProtocolStack",
+			"Type": "EditorAndProgram",
+			"ProgramAllowList": [ "LiveLinkHub" ],
+			"LoadingPhase": "Default"
 		}
 	],
 	"Plugins": []
diff --git a/Engine/Plugins/VirtualProduction/CaptureManager/Source/CaptureProtocolStack/CaptureProtocolStack.Build.cs b/Engine/Plugins/VirtualProduction/CaptureManager/Source/CaptureProtocolStack/CaptureProtocolStack.Build.cs
new file mode 100644
--- /dev/null
+++ b/Engine/Plugins/VirtualProduction/CaptureManager/Source/CaptureProtocolStack/CaptureProtocolStack.Build.cs
@@ -0,0 +1,41 @@
+// Copyright Epic Games, Inc. All Rights Reserved.
+
+using UnrealBuildTool;
+using System.IO;
+
+public class CaptureProtocolStack : ModuleRules
+{
+	protected bool BuildForDevelopment
+	{
+		get
+		{
+			// Check if source is available
+			string SourceFilesPath = Path.Combine(ModuleDirectory, "Private");
+			return Directory.Exists(SourceFilesPath) &&
+				   Directory.GetFiles(SourceFilesPath).Length > 0;
+		}
+	}
+
+	public CaptureProtocolStack(ReadOnlyTargetRules Target) : base(Target)
+	{
+		bUsePrecompiled = !BuildForDevelopment;
+		PCHUsage = PCHUsageMode.UseExplicitOrSharedPCHs;
+
+        PublicDependencyModuleNames.AddRange(new string[]
+		{
+            "Core",
+			"CaptureUtils"
+		});
+
+		PrivateDependencyModuleNames.AddRange(new string[]
+		{
+			"CoreUObject",
+			"Engine",
+            "Json",
+            "JsonUtilities",
+            "Sockets",
+            "Networking",
+            "UnrealEd"
+        });
+	}
+}
diff --git a/Engine/Plugins/VirtualProduction/CaptureManager/Source/CaptureProtocolStack/Private/CaptureProtocolStack.cpp b/Engine/Plugins/VirtualProduction/CaptureManager/Source/CaptureProtocolStack/Private/CaptureProtocolStack.cpp
new file mode 100644
--- /dev/null
+++ b/Engine/Plugins/VirtualProduction/CaptureManager/Source/CaptureProtocolStack/Private/CaptureProtocolStack.cpp
@@ -0,0 +1,20 @@
+// Copyright Epic Games, Inc. All Rights Reserved.
+
+#include "CaptureProtocolStack.h"
+
+void FCaptureProtocolStackModule::StartupModule()
+{
+	TimerManager = MakeUnique<FCPSTimerManager>();
+}
+
+void FCaptureProtocolStackModule::ShutdownModule()
+{
+	TimerManager = nullptr;
+}
+
+FCPSTimerManager& FCaptureProtocolStackModule::GetTimerManager()
+{
+	return *TimerManager;
+}
+
+IMPLEMENT_MODULE(FCaptureProtocolStackModule, CaptureProtocolStack)
diff --git a/Engine/Plugins/VirtualProduction/CaptureManager/Source/CaptureProtocolStack/Private/Communication/TcpClient.cpp b/Engine/Plugins/VirtualProduction/CaptureManager/Source/CaptureProtocolStack/Private/Communication/TcpClient.cpp
new file mode 100644
--- /dev/null
+++ b/Engine/Plugins/VirtualProduction/CaptureManager/Source/CaptureProtocolStack/Private/Communication/TcpClient.cpp
@@ -0,0 +1,185 @@
+// Copyright Epic Games, Inc. All Rights Reserved.
+
+#include "Communication/TcpClient.h"
+
+FTcpClient::FTcpClient() 
+    : bRunning(false)
+{
+}
+
+FTcpClient::~FTcpClient()
+{
+    Stop();
+}
+
+TProtocolResult<void> FTcpClient::Init()
+{
+    if (bRunning)
+    {
+        return FCaptureProtocolError(TEXT("Can't initialize the client while running."));
+    }
+
+    // Prepare socket
+	FSocket* RawSocket = FTcpSocketBuilder(TEXT("CPS TCP Socket"))
+		.AsNonBlocking()
+		.AsReusable()
+		.WithReceiveBufferSize(BufferSize)
+		.WithSendBufferSize(BufferSize)
+		.Build();
+
+	if (!RawSocket)
+	{
+		return FCaptureProtocolError(TEXT("Failed to create a client socket"));
+	}
+
+	// Wrapping the socket in UniquePtr
+    TcpSocket = FSocketPtr(RawSocket, FSocketDeleter(ISocketSubsystem::Get()));
+
+    return ResultOk;
+}
+
+// Blocking function until connection is established
+TProtocolResult<void> FTcpClient::Start(const FString& InServerAddress)
+{
+    if (bRunning)
+    {
+        return FCaptureProtocolError(TEXT("The client is already started"));
+    }
+
+    FIPv4Endpoint Endpoint;
+    FIPv4Endpoint::FromHostAndPort(InServerAddress, Endpoint);
+
+    bool IsConnected = TcpSocket->Connect(*Endpoint.ToInternetAddr());        
+
+    // Wait 500 milliseconds before reporting an error
+    FDateTime Time = FDateTime::UtcNow();
+    while (IsConnected && TcpSocket->GetConnectionState() != SCS_Connected)
+    {
+        FPlatformProcess::Sleep(0.20f);
+        if ((FDateTime::UtcNow() - Time).GetTotalMilliseconds() > 500)
+        {
+            IsConnected = false;
+        }
+    }
+
+    if (!IsConnected)
+    {
+        return FCaptureProtocolError(TEXT("Failed to connect the client"));
+    }
+
+    bRunning = true;
+
+    return ResultOk;
+}
+
+TProtocolResult<void> FTcpClient::Stop()
+{
+    if (!bRunning)
+    {
+        return FCaptureProtocolError(TEXT("The client is already stopped"));
+    }
+
+    TcpSocket->Shutdown(ESocketShutdownMode::ReadWrite);
+    TcpSocket->Close();
+    TcpSocket = nullptr;
+
+    bRunning = false;
+
+    return ResultOk;
+}
+
+bool FTcpClient::IsRunning() const
+{
+	return bRunning;
+}
+
+TProtocolResult<void> FTcpClient::SendMessage(const TArray<uint8>& InPayload)
+{
+    if (!TcpSocket)
+    {
+        return FCaptureProtocolError(TEXT("Invalid TCP socket"));
+    }
+
+	int32 TotalSent = 0;
+	while (TotalSent != InPayload.Num())
+	{
+		int32 Sent = 0;
+		bool Result = TcpSocket->Send(InPayload.GetData() + TotalSent, InPayload.Num() - TotalSent, Sent);
+		if (!Result)
+		{
+			return FCaptureProtocolError(TEXT("Failed to send the data"));
+		}
+
+		TotalSent += Sent;
+	}
+
+    return ResultOk;
+}
+
+TProtocolResult<TArray<uint8>> FTcpClient::ReceiveMessage(const uint64 InSize, const uint32 InWaitTimeoutMs)
+{
+    if (!TcpSocket)
+    {
+        return FCaptureProtocolError(TEXT("Invalid TCP socket"));
+    }
+
+	TArray<uint8> ReadData;
+	ReadData.AddUninitialized(InSize);
+
+	uint32 LeftToRead = InSize;
+
+	while (LeftToRead != 0)
+	{
+		uint32 PendingSize = 0;
+		int32 ReadSize = 0;
+
+		if (!TcpSocket->Wait(ESocketWaitConditions::WaitForRead, FTimespan::FromMilliseconds(InWaitTimeoutMs)))
+		{
+			return FCaptureProtocolError(TEXT("Timeout has expired"), TimeoutError);
+		}
+
+		if (!TcpSocket->HasPendingData(PendingSize))
+		{
+			return FCaptureProtocolError(TEXT("Host has been disconnected"), DisconnectedError);
+		}
+
+		int32 Size = FMath::Min(LeftToRead, PendingSize);
+
+		const uint32 Offset = InSize - LeftToRead;
+		if (!TcpSocket->Recv(ReadData.GetData() + Offset, Size, ReadSize))
+		{
+			return FCaptureProtocolError(TEXT("Failed to read the data from the TCP socket"));
+		}
+
+		if (PendingSize == 0 && ReadSize == 0)
+		{
+			return FCaptureProtocolError(TEXT("Host has been disconnected"), DisconnectedError);
+		}
+
+		LeftToRead -= ReadSize;
+	}
+
+    return ReadData;
+}
+
+FTcpClientReader::FTcpClientReader(FTcpClient& InClient)
+    : Client(InClient)
+{
+}
+
+TProtocolResult<TArray<uint8>> FTcpClientReader::ReceiveMessage(const uint64 InSize, const uint32 InWaitTimeoutMs)
+{
+    return Client.ReceiveMessage(InSize, InWaitTimeoutMs);
+}
+
+
+FTcpClientWriter::FTcpClientWriter(FTcpClient& InClient)
+    : Client(InClient)
+{
+}
+
+TProtocolResult<void> FTcpClientWriter::SendMessage(const TArray<uint8>& InPayload)
+{
+    return Client.SendMessage(InPayload);
+}
+
diff --git a/Engine/Plugins/VirtualProduction/CaptureManager/Source/CaptureProtocolStack/Private/Communication/TcpServer.cpp b/Engine/Plugins/VirtualProduction/CaptureManager/Source/CaptureProtocolStack/Private/Communication/TcpServer.cpp
new file mode 100644
--- /dev/null
+++ b/Engine/Plugins/VirtualProduction/CaptureManager/Source/CaptureProtocolStack/Private/Communication/TcpServer.cpp
@@ -0,0 +1,268 @@
+// Copyright Epic Games, Inc. All Rights Reserved.
+
+#include "Communication/TcpServer.h"
+
+FTcpClientHandler::FTcpClientHandler(FSocketPtr InSocket, FString InEndpoint)
+    : Socket(MoveTemp(InSocket))
+    , Endpoint(MoveTemp(InEndpoint))
+{
+}
+
+FTcpClientHandler::~FTcpClientHandler()
+{
+    if (Socket)
+    {
+        Socket->Shutdown(ESocketShutdownMode::ReadWrite);
+        Socket->Close();
+        Socket = nullptr;
+    }
+}
+
+TProtocolResult<void> FTcpClientHandler::SendMessage(const TArray<uint8>& InData)
+{
+    if (!Socket)
+    {
+        return FCaptureProtocolError(TEXT("Invalid TCP socket"));
+    }
+
+	int32 TotalSent = 0;
+	while (TotalSent != InData.Num())
+	{
+		int32 Sent = 0;
+		bool Result = Socket->Send(InData.GetData() + TotalSent, InData.Num() - TotalSent, Sent);
+		if (!Result)
+		{
+			return FCaptureProtocolError(TEXT("Failed to send the data"));
+		}
+
+		TotalSent += Sent;
+	}
+
+	return ResultOk;
+}
+
+TProtocolResult<TArray<uint8>> FTcpClientHandler::ReceiveMessage(const uint64 InSize, const uint32 InWaitTimeoutMs)
+{
+    if (!Socket)
+    {
+        return FCaptureProtocolError(TEXT("Invalid TCP socket"));
+    }
+
+	TArray<uint8> ReadData;
+	ReadData.AddUninitialized(InSize);
+
+	uint32 LeftToRead = InSize;
+
+	while (LeftToRead != 0)
+	{
+		uint32 PendingSize = 0;
+		int32 ReadSize = 0;
+
+		if (!Socket->Wait(ESocketWaitConditions::WaitForRead, FTimespan::FromMilliseconds(InWaitTimeoutMs)))
+		{
+			return FCaptureProtocolError(TEXT("Timeout has expired"), TimeoutError);
+		}
+
+		if (!Socket->HasPendingData(PendingSize))
+		{
+			return FCaptureProtocolError(TEXT("Host has been disconnected"), DisconnectedError);
+		}
+
+		int32 Size = FMath::Min(LeftToRead, PendingSize);
+
+		const uint32 Offset = InSize - LeftToRead;
+		if (!Socket->Recv(ReadData.GetData() + Offset, Size, ReadSize))
+		{
+			return FCaptureProtocolError(TEXT("Failed to read the data from the TCP socket"));
+		}
+
+		if (PendingSize == 0 && ReadSize == 0)
+		{
+			return FCaptureProtocolError(TEXT("Host has been disconnected"), DisconnectedError);
+		}
+
+		LeftToRead -= ReadSize;
+	}
+
+    return ReadData;
+}
+
+TProtocolResult<uint32> FTcpClientHandler::HasPendingData() const
+{
+	if (!Socket)
+	{
+		return FCaptureProtocolError(TEXT("Invalid TCP socket"));
+	}
+
+	uint32 PendingSize = 0;
+	if (!Socket->HasPendingData(PendingSize))
+	{
+		return FCaptureProtocolError(TEXT("Host has been disconnected"), DisconnectedError);
+	}
+
+	return PendingSize;
+}
+
+const FString& FTcpClientHandler::GetEndpoint() const
+{
+    return Endpoint;
+}
+
+bool FTcpClientHandler::operator==(const FTcpClientHandler& InOther)
+{
+    return InOther.Endpoint == Endpoint;
+}
+
+
+FTcpServer::FTcpServer(const uint32 InMaxNumberOfClients)
+    : Listener(nullptr)
+    , MaxNumberOfClients(InMaxNumberOfClients)
+    , bRunning(false)
+{
+    Clients.Reserve(MaxNumberOfClients);
+}
+
+TProtocolResult<uint16> FTcpServer::Start(const uint16 InListeningPort)
+{
+    if (bRunning)
+    {
+        return FCaptureProtocolError(TEXT("The server is already started"));
+    }
+
+	FIPv4Endpoint ListenEndpoint(FIPv4Address::Any, InListeningPort);
+
+	FSocket* RawPointer = FTcpSocketBuilder(TEXT("FTcpListener server"))
+		.AsReusable(true)
+		.BoundToEndpoint(ListenEndpoint)
+		.Listening(8) // Max number of connections in queue before refusing them
+		.WithSendBufferSize(2 * 1024 * 1024).Build();
+
+	if (!RawPointer)
+	{
+		return FCaptureProtocolError(TEXT("Failed to create a server socket"));
+	}
+
+	// Wrapping the socket in UniquePtr
+	Socket = FSocketPtr(RawPointer, FSocketDeleter(ISocketSubsystem::Get()));
+
+    Listener.Reset(new FTcpListener(*Socket.Get(), FTimespan::FromMilliseconds(ThreadWaitTime)));
+
+    if (!Listener)
+    {
+        return FCaptureProtocolError(TEXT("The server is failed to create a listener"));
+    }
+
+    Listener->OnConnectionAccepted().BindLambda([this](FSocket* InSocket, const FIPv4Endpoint& InEndpoint) mutable 
+    {
+        if (Clients.Num() >= (int32) MaxNumberOfClients)
+        {
+            return false;
+        }
+
+		FSocketPtr ClientSocket(InSocket, FSocketDeleter(ISocketSubsystem::Get()));
+
+		int32 NewSize = 0;
+		ClientSocket->SetReceiveBufferSize(FTcpClientHandler::MaxBufferSize, NewSize);
+
+        TSharedPtr<FTcpClientHandler> NewClient = MakeShared<FTcpClientHandler>(MoveTemp(ClientSocket), InEndpoint.ToString());
+
+        TWeakPtr<FTcpClientHandler> ClientObserver(NewClient);
+        OnConnectionHandler.ExecuteIfBound(MoveTemp(ClientObserver), true);
+
+        FRWScopeLock Lock(Mutex, FRWScopeLockType::SLT_Write);
+        Clients.Add(InEndpoint.ToString(), MoveTemp(NewClient));
+
+        return true;
+    });
+
+    bRunning = true;
+
+    return GetPort();
+}
+
+TProtocolResult<void> FTcpServer::Stop()
+{
+    if (!bRunning)
+    {
+        return FCaptureProtocolError(TEXT("The server is already stopped"));
+    }
+
+    Listener->Stop();
+    Listener = nullptr;
+
+	Socket->Close();
+	Socket = nullptr;
+
+    bRunning = false;
+
+    FRWScopeLock Lock(Mutex, FRWScopeLockType::SLT_Write);
+
+    Clients.Empty();
+
+    return ResultOk;
+}
+
+TProtocolResult<void> FTcpServer::SendMessage(const TArray<uint8>& InMessage, const FString& InEndpoint)
+{
+    FRWScopeLock Lock(Mutex, FRWScopeLockType::SLT_ReadOnly);
+
+    TSharedPtr<FTcpClientHandler>* Iterator = Clients.Find(InEndpoint);
+    verifyf(Iterator != nullptr, TEXT("Client doesn't exist"));
+
+    const TSharedPtr<FTcpClientHandler>& Client = *Iterator;
+
+    return Client->SendMessage(InMessage);
+}
+
+void FTcpServer::DisconnectClient(const FString& InEndpoint)
+{
+    FRWScopeLock Lock(Mutex, FRWScopeLockType::SLT_Write);
+
+    TSharedPtr<FTcpClientHandler>* Iterator = Clients.Find(InEndpoint);
+
+	if (Iterator)
+	{
+		OnConnectionHandler.ExecuteIfBound(*Iterator, false);
+		Clients.Remove(InEndpoint);
+	}
+}
+
+void FTcpServer::DisconnectClientAsync(const FString& InEndpoint)
+{
+    Async(EAsyncExecution::LargeThreadPool, [this, InEndpoint]()
+    {
+        DisconnectClient(InEndpoint);
+    });
+}
+
+void FTcpServer::SetConnectionHandler(FConnectionHandler InOnConnectionHandler)
+{
+    OnConnectionHandler = MoveTemp(InOnConnectionHandler);
+}
+
+int32 FTcpServer::GetPort() const
+{
+	return Socket->GetPortNo();
+}
+
+FTcpConnectionReader::FTcpConnectionReader(FTcpClientHandler& InClient)
+    : Client(InClient)
+{
+}
+
+TProtocolResult<TArray<uint8>> FTcpConnectionReader::ReceiveMessage(const uint64 InSize, const uint32 InWaitTimeoutMs)
+{
+    return Client.ReceiveMessage(InSize, InWaitTimeoutMs);
+}
+
+
+FTcpConnectionWriter::FTcpConnectionWriter(FTcpClientHandler& InClient)
+    : Client(InClient)
+{
+}
+
+TProtocolResult<void> FTcpConnectionWriter::SendMessage(const TArray<uint8>& InPayload)
+{
+    return Client.SendMessage(InPayload);
+}
+
diff --git a/Engine/Plugins/VirtualProduction/CaptureManager/Source/CaptureProtocolStack/Private/Communication/UdpClient.cpp b/Engine/Plugins/VirtualProduction/CaptureManager/Source/CaptureProtocolStack/Private/Communication/UdpClient.cpp
new file mode 100644
--- /dev/null
+++ b/Engine/Plugins/VirtualProduction/CaptureManager/Source/CaptureProtocolStack/Private/Communication/UdpClient.cpp
@@ -0,0 +1,109 @@
+// Copyright Epic Games, Inc. All Rights Reserved.
+
+#include "Communication/UdpClient.h"
+
+FUdpClient::FUdpClient() 
+    : bRunning(false)
+{
+}
+
+FUdpClient::~FUdpClient()
+{
+    Stop();
+}
+
+TProtocolResult<void> FUdpClient::Init(FUdpClientConfigure InConfig, FOnSocketDataReceived InReceiveHandler)
+{
+    if (bRunning)
+    {
+        return FCaptureProtocolError(TEXT("Can't initialize the client while running."));
+    }
+
+    FIPv4Endpoint Endpoint(FIPv4Address::Any, InConfig.ListenPort);
+
+    // Prepare socket
+    UdpSocket.Reset(FUdpSocketBuilder(TEXT("CPS UDP Socket"))
+                    .AsNonBlocking()
+                    .AsReusable()
+                    .BoundToEndpoint(Endpoint)
+                    .WithReceiveBufferSize(BufferSize)
+                    .WithSendBufferSize(BufferSize)
+                    .Build());
+
+    if (!UdpSocket)
+    {
+        return FCaptureProtocolError(TEXT("Failed to create a client socket"));
+    }
+
+	if (!InConfig.MulticastIpAddress.IsEmpty())
+	{
+		FIPv4Endpoint MulticastEndpoint;
+		FString MulticastHostAnDPort = FString::Format(TEXT("{0}:{1}"), { InConfig.MulticastIpAddress, InConfig.ListenPort });
+
+		CHECK_BOOL(FIPv4Endpoint::FromHostAndPort(MulticastHostAnDPort, MulticastEndpoint));
+        CHECK_BOOL(UdpSocket->JoinMulticastGroup(*MulticastEndpoint.ToInternetAddr()));
+        CHECK_BOOL(UdpSocket->SetMulticastLoopback(true));
+    }
+
+    // Prepare receiver
+    UdpReceiver.Reset(new FUdpSocketReceiver(UdpSocket.Get(), FTimespan::FromMilliseconds(ThreadWaitTime), TEXT("CPS RECEIVER-FUdpCommunication")));
+
+    UdpReceiver->OnDataReceived() = MoveTemp(InReceiveHandler);
+
+    return ResultOk;
+}
+
+TProtocolResult<void> FUdpClient::Start()
+{
+    if (bRunning)
+    {
+        return FCaptureProtocolError(TEXT("The client is already started"));
+    }
+
+    UdpReceiver->Start();
+
+    bRunning = true;
+        
+    return ResultOk;
+}
+
+TProtocolResult<void> FUdpClient::Stop()
+{
+    if (!bRunning)
+    {
+        return FCaptureProtocolError(TEXT("The client is already stopped"));
+    }
+
+    UdpReceiver->Stop();
+    UdpReceiver = nullptr;
+
+    UdpSocket->Close();
+    UdpSocket = nullptr;
+
+    bRunning = false;
+
+    return ResultOk;
+}
+
+TProtocolResult<int32> FUdpClient::SendMessage(const TArray<uint8>& InPayload, const FString& InEndpoint)
+{
+    if (!UdpSocket)
+    {
+        return FCaptureProtocolError(TEXT("Udp socket not configured"));
+    }
+
+    FIPv4Endpoint Endpoint;
+    FIPv4Endpoint::FromHostAndPort(InEndpoint, Endpoint);
+
+    TSharedRef<FInternetAddr> Address = Endpoint.ToInternetAddr();
+
+    int32 Sent = 0;
+    bool Result = UdpSocket->SendTo(InPayload.GetData(), InPayload.Num(), Sent, *Address);
+    if (!Result)
+    {
+        return FCaptureProtocolError(TEXT("Failed to send the data"));
+    }
+
+    return Sent;
+}
+
diff --git a/Engine/Plugins/VirtualProduction/CaptureManager/Source/CaptureProtocolStack/Private/Control/Communication/ControlCommunication.cpp b/Engine/Plugins/VirtualProduction/CaptureManager/Source/CaptureProtocolStack/Private/Control/Communication/ControlCommunication.cpp
new file mode 100644
--- /dev/null
+++ b/Engine/Plugins/VirtualProduction/CaptureManager/Source/CaptureProtocolStack/Private/Control/Communication/ControlCommunication.cpp
@@ -0,0 +1,160 @@
+// Copyright Epic Games, Inc. All Rights Reserved.
+
+#include "Control/Communication/ControlCommunication.h"
+
+#include "Utility/Definitions.h"
+
+DEFINE_LOG_CATEGORY(LogCPSControlCommunication)
+
+FControlCommunication::FControlCommunication()
+    : ReceiveHandler(nullptr)
+    , SynchronizedReceiver(FRunnerType::FOnProcess::CreateRaw(this, &FControlCommunication::OnProcessReceivedPacket))
+    , SynchronizedSender(FRunnerType::FOnProcess::CreateRaw(this, &FControlCommunication::OnProcessSentPacket))
+	, SynchronizedRunnable(TQueueRunner<TSharedPtr<FRunnable>>::FOnProcess::CreateRaw(this, &FControlCommunication::RunnableHandler))
+{
+}
+
+TProtocolResult<void> FControlCommunication::Init()
+{
+    return Client.Init();
+}
+
+TProtocolResult<void> FControlCommunication::Start(const FString& InServerIp, const uint16 InServerPort)
+{
+	TProtocolResult<void> Result = Client.Start(InServerIp + TEXT(":") + FString::FromInt(InServerPort));
+	if (Result.HasValue())
+	{
+		TSharedPtr<FCommunicationRunnable> CommunicationRunnable = MakeShared<FCommunicationRunnable>(*this);
+		Runnable = CommunicationRunnable;
+
+		SynchronizedRunnable.Add(CommunicationRunnable);
+	}
+
+    return Result;
+}
+
+TProtocolResult<void> FControlCommunication::Stop()
+{
+	if (TSharedPtr<FCommunicationRunnable> CommunicationRunnable = Runnable.Pin(); CommunicationRunnable)
+	{
+		CommunicationRunnable->Stop();
+		CommunicationRunnable->Join();
+		Runnable.Reset();
+	}
+
+    return Client.Stop();
+}
+
+bool FControlCommunication::IsRunning() const
+{
+	return Client.IsRunning();
+}
+
+void FControlCommunication::SendMessage(FControlPacket InMessage)
+{
+    SynchronizedSender.Add(MoveTemp(InMessage));
+}
+
+void FControlCommunication::SetReceiveHandler(FOnPacketReceived InReceiveHandler)
+{
+    ReceiveHandler = MoveTemp(InReceiveHandler);
+}
+
+void FControlCommunication::OnProcessReceivedPacket(FControlPacket InMessage)
+{
+    ReceiveHandler.ExecuteIfBound(MoveTemp(InMessage));
+}
+
+void FControlCommunication::OnProcessSentPacket(FControlPacket InMessage)
+{
+    FTcpClientWriter Writer(Client);
+
+    TProtocolResult<void> SerializeResult = FControlPacket::Serialize(InMessage, Writer);
+    if (SerializeResult.HasError())
+    {
+        UE_LOG(LogCPSControlCommunication, Error, TEXT("Invalid message: %s"), *(SerializeResult.StealError().GetMessage()));
+        return;
+    }
+}
+
+void FControlCommunication::RunnableHandler(TSharedPtr<FRunnable> InRunnable)
+{
+	InRunnable->Run();
+}
+
+TProtocolResult<FControlPacketHeader> FControlCommunication::ReceiveControlHeader()
+{
+    FTcpClientReader Reader(Client);
+
+    return FControlPacketHeader::Deserialize(Reader);
+}
+
+TProtocolResult<void> FControlCommunication::ReceiveControlPacket(const FControlPacketHeader& InHeader)
+{
+    FTcpClientReader Reader(Client);
+
+    TProtocolResult<FControlPacket> PacketDeserialize = FControlPacket::Deserialize(InHeader, Reader);
+    if (PacketDeserialize.HasError())
+    {
+        return PacketDeserialize.StealError();
+    }
+
+    SynchronizedReceiver.Add(PacketDeserialize.StealValue());
+
+    return ResultOk;
+}
+
+FControlCommunication::FCommunicationRunnable::FCommunicationRunnable(FControlCommunication& InCommunication)
+	: Communication(InCommunication)
+	, bIsRunning(true)
+{
+	DoneEvent = FGenericPlatformProcess::GetSynchEventFromPool(false);
+	check(DoneEvent);
+}
+
+FControlCommunication::FCommunicationRunnable::~FCommunicationRunnable()
+{
+	FGenericPlatformProcess::ReturnSynchEventToPool(DoneEvent);
+}
+
+uint32 FControlCommunication::FCommunicationRunnable::Run()
+{
+	while (bIsRunning)
+	{
+		TProtocolResult<FControlPacketHeader> ReceiveHeaderResult = Communication.ReceiveControlHeader();
+		if (ReceiveHeaderResult.HasError())
+		{
+			FCaptureProtocolError Error = ReceiveHeaderResult.StealError();
+			if (Error.GetCode() == FTcpClient::DisconnectedError)
+			{
+				Stop();
+			}
+
+			continue;
+		}
+
+		TProtocolResult<void> ReceivePacketResult = Communication.ReceiveControlPacket(ReceiveHeaderResult.StealValue());
+		if (ReceivePacketResult.HasError())
+		{
+			FCaptureProtocolError Error = ReceivePacketResult.StealError();
+			if (Error.GetCode() == FTcpClient::DisconnectedError)
+			{
+				Stop();
+			}
+		}
+	}
+
+	DoneEvent->Trigger();
+
+	return 0;
+}
+
+void FControlCommunication::FCommunicationRunnable::Stop()
+{
+	bIsRunning = false;
+}
+
+void FControlCommunication::FCommunicationRunnable::Join()
+{
+	DoneEvent->Wait();
+}
\ No newline at end of file
diff --git a/Engine/Plugins/VirtualProduction/CaptureManager/Source/CaptureProtocolStack/Private/Control/Communication/ControlPacket.cpp b/Engine/Plugins/VirtualProduction/CaptureManager/Source/CaptureProtocolStack/Private/Control/Communication/ControlPacket.cpp
new file mode 100644
--- /dev/null
+++ b/Engine/Plugins/VirtualProduction/CaptureManager/Source/CaptureProtocolStack/Private/Control/Communication/ControlPacket.cpp
@@ -0,0 +1,117 @@
+// Copyright Epic Games, Inc. All Rights Reserved.
+
+#include "Control/Communication/ControlPacket.h"
+
+#include "Control/Messages/ControlJsonUtilities.h"
+
+const TArray<uint8> FControlPacketHeader::Header = { 'C', 'P', 'S', 'C', 'O', 'N', 'T', 'R', 'O', 'L', '\0' };
+
+#define CPS_CHECK_RESULT_RET_ERROR(Result) if (Result.HasError()) { return Result.StealError(); }
+
+TProtocolResult<FControlPacketHeader> FControlPacketHeader::Deserialize(ITcpSocketReader& InReader)
+{
+    TProtocolResult<TArray<uint8>> HeaderResult = InReader.ReceiveMessage(Header.Num());
+    CPS_CHECK_RESULT_RET_ERROR(HeaderResult);
+
+    TArray<uint8> HeaderData = HeaderResult.StealValue();
+
+    if (FMemory::Memcmp(HeaderData.GetData(), Header.GetData(), Header.Num()) != 0)
+    {
+        return FCaptureProtocolError(TEXT("Header doesn't match"));
+    }
+
+	TProtocolResult<TArray<uint8>> VersionResult = InReader.ReceiveMessage(sizeof(uint16));
+	CPS_CHECK_RESULT_RET_ERROR(VersionResult);
+
+	TArray<uint8> VersionData = VersionResult.StealValue();
+
+	TProtocolResult<TArray<uint8>> PayloadSizeResult = InReader.ReceiveMessage(sizeof(uint32));
+    CPS_CHECK_RESULT_RET_ERROR(PayloadSizeResult);
+
+    TArray<uint8> PayloadSizeData = PayloadSizeResult.StealValue();
+
+    FControlPacketHeader PacketHeader;
+	PacketHeader.Version = *(reinterpret_cast<const uint16*>(VersionData.GetData()));
+    PacketHeader.PayloadSize = *(reinterpret_cast<const uint32*>(PayloadSizeData.GetData()));
+
+    return PacketHeader;
+}
+
+uint32 FControlPacketHeader::GetHeaderSize()
+{
+    return Header.Num() + sizeof(uint16) + sizeof(uint32);
+}
+
+FControlPacketHeader::FControlPacketHeader()
+{
+}
+
+FControlPacketHeader::FControlPacketHeader(uint16 InVersion, uint32 InPayloadSize)
+	: Version(InVersion)
+	, PayloadSize(InPayloadSize)
+{
+}
+
+uint16 FControlPacketHeader::GetVersion() const
+{
+	return Version;
+}
+
+uint32 FControlPacketHeader::GetPayloadSize() const
+{
+    return PayloadSize;
+}
+
+TProtocolResult<FControlPacket> FControlPacket::Deserialize(const FControlPacketHeader& InPacketHeader, ITcpSocketReader& InReader)
+{
+    TProtocolResult<TArray<uint8>> PayloadResult = InReader.ReceiveMessage(InPacketHeader.GetPayloadSize());
+    CPS_CHECK_RESULT_RET_ERROR(PayloadResult);
+
+    FControlPacket Packet;
+
+	Packet.Header = InPacketHeader;
+	Packet.Payload = PayloadResult.StealValue();
+
+    return Packet;
+}
+
+TProtocolResult<void> FControlPacket::Serialize(const FControlPacket& InMessage, ITcpSocketWriter& InWriter)
+{
+    TArray<uint8> Data;
+    Data.Append(FControlPacketHeader::Header.GetData(), FControlPacketHeader::Header.Num());
+
+	uint16 Version = InMessage.GetVersion();
+	Data.Append(reinterpret_cast<uint8*>(&Version), sizeof(Version));
+
+    uint32 PayloadLength = InMessage.Payload.Num();
+
+    Data.Append(reinterpret_cast<uint8*>(&PayloadLength), sizeof(PayloadLength));
+	Data.Append(InMessage.Payload.GetData(), PayloadLength);
+
+	return InWriter.SendMessage(Data);
+}
+
+FControlPacket::FControlPacket()
+{
+}
+
+FControlPacket::FControlPacket(uint16 InVersion, TArray<uint8> InPayload)
+    : Header(FControlPacketHeader(InVersion, InPayload.Num()))
+	, Payload(MoveTemp(InPayload))
+{
+}
+
+uint16 FControlPacket::GetVersion() const
+{
+	return Header.GetVersion();
+}
+
+const TArray<uint8>& FControlPacket::GetPayload() const
+{
+    return Payload;
+}
+
+uint32 FControlPacket::GetPayloadSize() const
+{
+    return Header.GetPayloadSize();
+}
\ No newline at end of file
diff --git a/Engine/Plugins/VirtualProduction/CaptureManager/Source/CaptureProtocolStack/Private/Control/ControlMessenger.cpp b/Engine/Plugins/VirtualProduction/CaptureManager/Source/CaptureProtocolStack/Private/Control/ControlMessenger.cpp
new file mode 100644
--- /dev/null
+++ b/Engine/Plugins/VirtualProduction/CaptureManager/Source/CaptureProtocolStack/Private/Control/ControlMessenger.cpp
@@ -0,0 +1,267 @@
+// Copyright Epic Games, Inc. All Rights Reserved.
+
+#include "Control/ControlMessenger.h"
+
+#include "CaptureProtocolStack.h"
+
+#include "Async/Async.h"
+#include "Editor.h"
+
+DEFINE_LOG_CATEGORY(LogCPSControlMessenger)
+
+
+FKeepAliveCounter::FKeepAliveCounter() :
+	Counter(0)
+{
+}
+
+void FKeepAliveCounter::Increment()
+{
+	++Counter;
+}
+
+void FKeepAliveCounter::Reset()
+{
+	Counter.store(0);
+}
+
+bool FKeepAliveCounter::HasReached(uint16 InBound)
+{
+	uint16 CurrentCounter = Counter.exchange(0);
+
+	if (CurrentCounter == InBound)
+	{
+		return true;
+	}
+
+	Counter.exchange(CurrentCounter);
+
+	return false;
+}
+
+const TCHAR* FControlMessenger::HandshakeSessionId = TEXT("handshake");
+FControlMessenger::FControlMessenger()
+    : SessionId(HandshakeSessionId)
+	, AsyncRequestRunner(TQueueRunner<FAsyncRequestDelegate>::FOnProcess::CreateRaw(this, &FControlMessenger::OnAsyncRequestProcess))
+	, RandomStream(GetTimestamp())
+{
+}
+
+FControlMessenger::~FControlMessenger()
+{
+    Stop();
+}
+
+void FControlMessenger::RegisterUpdateHandler(FString InAddressPath, FControlUpdate::FOnUpdateMessage InUpdateHandler)
+{
+    FScopeLock Lock(&UpdatesMutex);
+    UpdateHandlers.Emplace(MoveTemp(InAddressPath), MoveTemp(InUpdateHandler));
+}
+
+void FControlMessenger::RegisterDisconnectHandler(FOnDisconnect InOnDisconnectHandler)
+{
+	OnDisconnectHandler = MoveTemp(InOnDisconnectHandler);
+}
+
+TProtocolResult<void> FControlMessenger::Start(const FString& InServerIp, const uint16 InServerPort)
+{
+	if (!Communication.IsRunning())
+	{
+		CPS_CHECK_VOID_RESULT(Communication.Init());
+
+		Communication.SetReceiveHandler(FControlCommunication::FOnPacketReceived::CreateRaw(this, &FControlMessenger::MessageHandler));
+
+		CPS_CHECK_VOID_RESULT(Communication.Start(InServerIp, InServerPort));
+	}
+
+    return ResultOk;
+}
+
+TProtocolResult<void> FControlMessenger::Stop()
+{
+	if (Communication.IsRunning()) {
+
+		CPS_CHECK_VOID_RESULT(Communication.Stop());
+
+		if (KeepAliveTimer.IsValid())
+		{
+			StopKeepAliveTimer();
+		}
+	}
+
+    return ResultOk;
+}
+
+TProtocolResult<void> FControlMessenger::StartSession()
+{
+    TProtocolResult<FStartSessionResponse> Response = SendRequest(FStartSessionRequest());
+
+    if (Response.HasError())
+    {
+        return FCaptureProtocolError(TEXT("Response for Start Session Request is invalid."));
+    }
+
+    FScopeLock Lock(&SessionIdMutex);
+    if (SessionId != Response.GetValue().GetSessionId())
+    {
+        SessionId = Response.GetValue().GetSessionId();
+
+        if (KeepAliveTimer.IsValid())
+        {
+            StopKeepAliveTimer();
+        }
+
+        StartKeepAliveTimer();     
+    }
+
+    return ResultOk;
+}
+
+TProtocolResult<FGetServerInformationResponse> FControlMessenger::GetServerInformation()
+{
+    return SendRequest(FGetServerInformationRequest());
+}
+
+void FControlMessenger::SendPacket(FControlPacket InPacket)
+{
+    Communication.SendMessage(MoveTemp(InPacket));
+}
+
+void FControlMessenger::KeepAlive()
+{
+	SendAsyncRequest(FKeepAliveRequest(), 
+					 FOnControlResponse<FKeepAliveRequest>::CreateLambda([this](TProtocolResult<FKeepAliveResponse> InResult)
+	{
+		if (InResult.HasError())
+		{
+			KeepAliveFailures.Increment();
+
+			if (KeepAliveFailures.HasReached(3))
+			{
+				UE_LOG(LogCPSControlMessenger, Warning, TEXT("Server disconnected."));
+
+				{ 
+					FScopeLock Lock(&SessionIdMutex);
+					SessionId = HandshakeSessionId;
+				}
+
+				Stop();
+
+				OnDisconnectHandler.ExecuteIfBound(TEXT("Server failed to respond to Keep Alive message"));
+			}
+		}
+		else
+		{
+			KeepAliveFailures.Reset();
+		}
+	}));
+}
+
+void FControlMessenger::MessageHandler(FControlPacket InPacket)
+{
+    TProtocolResult<FControlMessage> DeserializeResult = FControlMessage::Deserialize(InPacket);
+
+    if (DeserializeResult.HasError())
+    {
+        FCaptureProtocolError DeserializeError = DeserializeResult.StealError();
+        UE_LOG(LogCPSControlMessenger, Error, TEXT("Failed to parse: %s"), *DeserializeError.GetMessage());
+        return;
+    }
+
+    FControlMessage Message = DeserializeResult.StealValue();
+
+    if (Message.GetType() == FControlMessage::EType::Request)
+    {
+        UE_LOG(LogCPSControlMessenger, Error, TEXT("Client currently doesn't support requests."));
+        return;
+    }
+    else if (Message.GetType() == FControlMessage::EType::Response)
+    {
+        FScopeLock Lock(&RequestsMutex);
+
+        if (const TUniquePtr<FRequestContext>* Iterator = RequestContexts.Find(Message.GetTransactionId()))
+        {
+            const TUniquePtr<FRequestContext>& RequestContext = *Iterator;
+
+            if (Message.GetAddressPath() != RequestContext->Request.GetAddressPath())
+            {
+                UE_LOG(LogCPSControlMessenger, Error, TEXT("Invalid response arrived"));
+                return;
+            }
+
+            RequestContext->Promise.SetValue(MoveTemp(Message));
+        }
+        return;
+    }
+    else if (Message.GetType() == FControlMessage::EType::Update)
+    {
+        FScopeLock Lock(&UpdatesMutex);
+
+        if (const FControlUpdate::FOnUpdateMessage* Iterator = UpdateHandlers.Find(Message.GetAddressPath()))
+        {
+            const FControlUpdate::FOnUpdateMessage& Handler = *Iterator;
+
+            // Using Shared Pointer as ExecuteIfBound can't accept non-copyable type
+            TProtocolResult<TSharedRef<FControlUpdate>> UpdateCreateResult = FControlUpdateCreator::Create(Message.GetAddressPath());
+
+            if (UpdateCreateResult.HasError())
+            {
+                UE_LOG(LogCPSControlMessenger, Error, TEXT("%s"), *(UpdateCreateResult.StealError().GetMessage()));
+                return;
+            }
+
+            TSharedPtr<FControlUpdate> Update = UpdateCreateResult.StealValue();
+
+            TProtocolResult<void> ParseResult = Update->Parse(Message.GetBody());
+            if (ParseResult.HasError())
+            {
+                UE_LOG(LogCPSControlMessenger, Error, TEXT("Failed to parse update: %s"), *(ParseResult.StealError().GetMessage()));
+                return;
+            }
+
+            Handler.ExecuteIfBound(MoveTemp(Update));
+        }
+
+        return;
+    }
+    else
+    {
+        UE_LOG(LogCPSControlMessenger, Error, TEXT("Invalid message arrived"));
+        return;
+    }
+}
+
+uint32 FControlMessenger::GenerateTransactionId() const
+{
+    float Fraction = RandomStream.GetFraction();
+
+    uint32 RandomNumber = Fraction * MAX_uint32;
+    return RandomNumber;
+}
+
+uint64 FControlMessenger::GetTimestamp() const
+{
+    FDateTime Now = FDateTime::UtcNow();
+    FDateTime Epoch(1970, 1, 1);
+
+    return static_cast<uint64>((Now - Epoch).GetTotalMilliseconds());
+}
+
+void FControlMessenger::StartKeepAliveTimer()
+{
+	FCPSTimerManager& TimerManager = FModuleManager::LoadModuleChecked<FCaptureProtocolStackModule>(TEXT("CaptureProtocolStack")).GetTimerManager();
+
+	KeepAliveTimer = TimerManager.AddTimer(FTimerDelegate::CreateRaw(this, &FControlMessenger::KeepAlive), KeepAliveInterval, true, KeepAliveInterval);
+}
+
+void FControlMessenger::StopKeepAliveTimer()
+{
+	FCPSTimerManager& TimerManager = FModuleManager::LoadModuleChecked<FCaptureProtocolStackModule>(TEXT("CaptureProtocolStack")).GetTimerManager();
+
+	TimerManager.RemoveTimer(KeepAliveTimer);
+}
+
+void FControlMessenger::OnAsyncRequestProcess(FAsyncRequestDelegate InAsyncDelegate)
+{
+	InAsyncDelegate.ExecuteIfBound();
+}
diff --git a/Engine/Plugins/VirtualProduction/CaptureManager/Source/CaptureProtocolStack/Private/Control/Messages/Constants.cpp b/Engine/Plugins/VirtualProduction/CaptureManager/Source/CaptureProtocolStack/Private/Control/Messages/Constants.cpp
new file mode 100644
--- /dev/null
+++ b/Engine/Plugins/VirtualProduction/CaptureManager/Source/CaptureProtocolStack/Private/Control/Messages/Constants.cpp
@@ -0,0 +1,127 @@
+// Copyright Epic Games, Inc. All Rights Reserved.
+
+#include "Control/Messages/Constants.h"
+
+namespace UE::CPS::AddressPaths
+{
+// Requests/Responses
+const TCHAR* GKeepAlive = TEXT("/session/keep-alive");
+const TCHAR* GStartSession = TEXT("/session/start");
+const TCHAR* GStopSession = TEXT("/session/stop");
+const TCHAR* GSubscribe = TEXT("/subscribe");
+const TCHAR* GUnsubscribe = TEXT("/unsubscribe");
+const TCHAR* GGetServerInformation = TEXT("/information");
+const TCHAR* GGetState = TEXT("/state");
+const TCHAR* GStartRecordingTake = TEXT("/take/record/start");
+const TCHAR* GStopRecordingTake = TEXT("/take/record/stop");
+const TCHAR* GAbortRecordingTake = TEXT("/take/record/abort");
+const TCHAR* GGetTakeList = TEXT("/take/list");
+const TCHAR* GGetTakeMetadata = TEXT("/take/metadata");
+const TCHAR* GExportTakeData = TEXT("/export/start/take");
+const TCHAR* GExportTakeVideoFrame = TEXT("/export/start/take/frame");
+const TCHAR* GExportCameraFeedFrame = TEXT("/export/start/camera/frame");
+const TCHAR* GPauseExport = TEXT("/export/pause");
+const TCHAR* GCancelExport = TEXT("/export/cancel");
+
+// Updates
+const TCHAR* GSessionStopped = TEXT("/event/session-stopped");
+const TCHAR* GRecordingStatus = TEXT("/state/recording");
+const TCHAR* GTakeAdded = TEXT("/take/added");
+const TCHAR* GTakeRemoved = TEXT("/take/removed");
+const TCHAR* GTakeUpdated = TEXT("/take/updated");
+
+// Platform Specific: IOS
+const TCHAR* GDiskCapacity = TEXT("state/ios/disk-capacity");
+const TCHAR* GBattery = TEXT("state/ios/battery");
+const TCHAR* GThermalState = TEXT("state/ios/thermal");
+}
+
+namespace UE::CPS::Properties
+{
+// Keys
+// Messages
+const TCHAR* GSessionId = TEXT("sessionId");
+const TCHAR* GAddressPath = TEXT("addressPath");
+const TCHAR* GTransactionId = TEXT("transactionId");
+const TCHAR* GTimestamp = TEXT("timestamp");
+const TCHAR* GType = TEXT("type");
+const TCHAR* GBody = TEXT("body");
+const TCHAR* GError = TEXT("error");
+const TCHAR* GDescription = TEXT("description");
+
+// Requests
+const TCHAR* GSlateName = TEXT("slateName");
+const TCHAR* GTakeNumber = TEXT("takeNumber");
+const TCHAR* GSubject = TEXT("subject");
+const TCHAR* GScenario = TEXT("scenario");
+const TCHAR* GTags = TEXT("tags");
+const TCHAR* GNames = TEXT("names");
+const TCHAR* GTakeName = TEXT("takeName");
+const TCHAR* GFiles = TEXT("files");
+const TCHAR* GName = TEXT("name");
+const TCHAR* GOffset = TEXT("offset");
+
+// Responses
+const TCHAR* GId = TEXT("id");
+const TCHAR* GModel = TEXT("model");
+const TCHAR* GPlatformName = TEXT("platformName");
+const TCHAR* GPlatformVersion = TEXT("platformVersion");
+const TCHAR* GSoftwareName = TEXT("softwareName");
+const TCHAR* GSoftwareVersion = TEXT("softwareVersion");
+const TCHAR* GExportPort = TEXT("exportPort");
+const TCHAR* GIsRecording = TEXT("isRecording");
+const TCHAR* GPlatformState = TEXT("platformState");
+const TCHAR* GTakes = TEXT("takes");
+const TCHAR* GDateTime = TEXT("dateTime");
+const TCHAR* GFrames = TEXT("frames");
+const TCHAR* GAppVersion = TEXT("appVersion");
+const TCHAR* GLength = TEXT("length");
+const TCHAR* GVideo = TEXT("video");
+const TCHAR* GFrameRate = TEXT("frameRate");
+const TCHAR* GHeight = TEXT("height");
+const TCHAR* GWidth = TEXT("width");
+const TCHAR* GAudio = TEXT("audio");
+const TCHAR* GChannels = TEXT("channels");
+const TCHAR* GSampleRate = TEXT("sampleRate");
+const TCHAR* GBitsPerChannel = TEXT("bitsPerChannel");
+
+// Platform Specific: IOS
+const TCHAR* GTotalCapacity = TEXT("totalCapacity");
+const TCHAR* GRemainingCapacity = TEXT("remainingCapacity");
+const TCHAR* GBatteryLevel = TEXT("batteryLevel");
+const TCHAR* GThermalState = TEXT("thermalState");
+
+const TCHAR* GTotal = TEXT("total");
+const TCHAR* GRemaining = TEXT("remaining");
+const TCHAR* GLevel = TEXT("level");
+const TCHAR* GState = TEXT("state");
+
+// Values
+// Messages
+const TCHAR* GRequest = TEXT("request");
+const TCHAR* GResponse = TEXT("response");
+const TCHAR* GUpdate = TEXT("update");
+
+const TCHAR* GIOS = TEXT("iOS");
+const TCHAR* GAndroid = TEXT("android");
+const TCHAR* GWindows = TEXT("windows");
+const TCHAR* GWindowsServer = TEXT("windowsServer");
+const TCHAR* GLinux = TEXT("linux");
+const TCHAR* GMacOS = TEXT("macOS");
+
+// Platform Specific: iOS
+const TCHAR* GNominal = TEXT("nominal");
+const TCHAR* GFair = TEXT("fair");
+const TCHAR* GSerious = TEXT("serious");
+const TCHAR* GCritical = TEXT("critical");
+}
+
+namespace UE::CPS::ErrorNames
+{
+const TCHAR* GUnsupportedProtocolVersion = TEXT("unsupportedProtocolVersion");
+const TCHAR* GRequestFailed = TEXT("requestFailed");
+const TCHAR* GRequestParsingFailed = TEXT("parsingRequestFailed");
+const TCHAR* GNotSupported = TEXT("notSupported");
+const TCHAR* GNoSession = TEXT("noSession");
+const TCHAR* GMissingBody = TEXT("missingBody");
+}
\ No newline at end of file
diff --git a/Engine/Plugins/VirtualProduction/CaptureManager/Source/CaptureProtocolStack/Private/Control/Messages/ControlJsonUtilities.cpp b/Engine/Plugins/VirtualProduction/CaptureManager/Source/CaptureProtocolStack/Private/Control/Messages/ControlJsonUtilities.cpp
new file mode 100644
--- /dev/null
+++ b/Engine/Plugins/VirtualProduction/CaptureManager/Source/CaptureProtocolStack/Private/Control/Messages/ControlJsonUtilities.cpp
@@ -0,0 +1,137 @@
+// Copyright Epic Games, Inc. All Rights Reserved.
+
+#include "Control/Messages/ControlJsonUtilities.h"
+
+class FArrayDataWriter final : public FArchive
+{
+public:
+
+	FArrayDataWriter(TArray<uint8>& InWriterArray)
+		: WriterArray(InWriterArray)
+	{
+	}
+
+	virtual ~FArrayDataWriter() = default;
+
+	virtual void Serialize(void* Data, int64 Num)
+	{
+		WriterArray.Append(static_cast<uint8*>(Data), Num);
+	}
+
+	virtual FString GetArchiveName() const
+	{
+		return TEXT("FArrayDataWriter");
+	}
+
+private:
+
+	TArray<uint8>& WriterArray;
+};
+
+class FArrayDataReader final : public FArchive
+{
+public:
+
+	FArrayDataReader(const TArray<uint8>& InReaderArray)
+		: ReaderArray(InReaderArray)
+		, ReaderPos(0)
+	{
+	}
+
+	virtual ~FArrayDataReader() = default;
+
+	virtual void Serialize(void* Data, int64 Num)
+	{
+		FMemory::Memcpy(Data, ReaderArray.GetData() + ReaderPos, Num);
+		ReaderPos += Num;
+	}
+
+	virtual int64 Tell()
+	{
+		return ReaderPos;
+	}
+
+	virtual int64 TotalSize()
+	{
+		return ReaderArray.Num();
+	}
+
+	virtual void Seek(int64 InPos)
+	{
+		ReaderPos = InPos;
+	}
+
+	virtual bool AtEnd()
+	{
+		return ReaderPos >= TotalSize();
+	}
+
+	virtual FString GetArchiveName() const
+	{
+		return TEXT("FArrayDataReader");
+	}
+
+private:
+
+	const TArray<uint8>& ReaderArray;
+	int64 ReaderPos;
+};
+
+bool FJsonUtility::CreateJsonFromUTF8Data(const TArray<uint8>& InData, TSharedPtr<FJsonObject>& OutObject)
+{
+	FArrayDataReader Reader(InData);
+
+	TSharedRef<TJsonReader<UTF8CHAR>> JsonReader = TJsonReaderFactory<UTF8CHAR>::Create(&Reader);
+	return FJsonSerializer::Deserialize(JsonReader, OutObject);
+}
+
+bool FJsonUtility::CreateUTF8DataFromJson(TSharedPtr<FJsonObject> InObject, TArray<uint8>& OutData)
+{
+	FArrayDataWriter Writer(OutData);
+
+	using FJsonWriterCondensedFactory = TJsonWriterFactory<UTF8CHAR, TCondensedJsonPrintPolicy<UTF8CHAR>>;
+	using FJsonWriterCondensed = TJsonWriter<UTF8CHAR, TCondensedJsonPrintPolicy<UTF8CHAR>>;
+
+	TSharedRef<FJsonWriterCondensed> JsonWriter = FJsonWriterCondensedFactory::Create(&Writer);
+	return FJsonSerializer::Serialize(InObject.ToSharedRef(), JsonWriter);
+}
+
+TProtocolResult<void> FJsonUtility::ParseString(const TSharedPtr<FJsonObject>& InBody, const FString& InFieldName, FString& OutFieldValue)
+{
+    if (!InBody->TryGetStringField(InFieldName, OutFieldValue))
+    {
+        return FCaptureProtocolError(TEXT("Failed to parse key: ") + InFieldName);
+    }
+
+    return ResultOk;
+}
+
+TProtocolResult<void> FJsonUtility::ParseBool(const TSharedPtr<FJsonObject>& InBody, const FString& InFieldName, bool& OutFieldValue)
+{
+    if (!InBody->TryGetBoolField(InFieldName, OutFieldValue))
+    {
+        return FCaptureProtocolError(TEXT("Failed to parse key: ") + InFieldName);
+    }
+
+    return ResultOk;
+}
+
+TProtocolResult<void> FJsonUtility::ParseObject(const TSharedPtr<FJsonObject>& InBody, const FString& InFieldName, const TSharedPtr<FJsonObject>*& OutFieldValue)
+{
+    if (!InBody->TryGetObjectField(InFieldName, OutFieldValue))
+    {
+        return FCaptureProtocolError(TEXT("Failed to parse key: ") + InFieldName);
+    }
+
+    return ResultOk;
+}
+
+TProtocolResult<void> FJsonUtility::ParseArray(const TSharedPtr<FJsonObject>& InBody, const FString& InFieldName, const TArray<TSharedPtr<FJsonValue>>*& OutFieldValue)
+{
+    if (!InBody->TryGetArrayField(InFieldName, OutFieldValue))
+    {
+        return FCaptureProtocolError(TEXT("Failed to parse key: ") + InFieldName);
+    }
+
+    return ResultOk;
+}
\ No newline at end of file
diff --git a/Engine/Plugins/VirtualProduction/CaptureManager/Source/CaptureProtocolStack/Private/Control/Messages/ControlMessage.cpp b/Engine/Plugins/VirtualProduction/CaptureManager/Source/CaptureProtocolStack/Private/Control/Messages/ControlMessage.cpp
new file mode 100644
--- /dev/null
+++ b/Engine/Plugins/VirtualProduction/CaptureManager/Source/CaptureProtocolStack/Private/Control/Messages/ControlMessage.cpp
@@ -0,0 +1,224 @@
+// Copyright Epic Games, Inc. All Rights Reserved.
+
+#include "Control/Messages/ControlMessage.h"
+#include "Control/Messages/Constants.h"
+
+#include "Control/Messages/ControlJsonUtilities.h"
+
+#include "Utility/Definitions.h"
+
+TProtocolResult<FControlMessage> FControlMessage::Deserialize(const FControlPacket& InPacket)
+{
+	if (InPacket.GetVersion() != CPS_VERSION)
+	{
+		return FCaptureProtocolError(TEXT("Version not supported"));
+	}
+
+	const TArray<uint8>& PayloadData = InPacket.GetPayload();
+
+	TSharedPtr<FJsonObject> Payload;
+	if (!FJsonUtility::CreateJsonFromUTF8Data(PayloadData, Payload))
+	{
+		return FCaptureProtocolError(TEXT("Failed to parse the data"));
+	}
+
+    FString SessionId;
+    CHECK_PARSE(FJsonUtility::ParseString(Payload, UE::CPS::Properties::GSessionId, SessionId));
+
+    FString AddressPath;
+    CHECK_PARSE(FJsonUtility::ParseString(Payload, UE::CPS::Properties::GAddressPath, AddressPath));
+
+    uint32 TransactionId;
+    CHECK_PARSE(FJsonUtility::ParseNumber(Payload, UE::CPS::Properties::GTransactionId, TransactionId));
+
+    uint64 Timestamp;
+    CHECK_PARSE(FJsonUtility::ParseNumber(Payload, UE::CPS::Properties::GTimestamp, Timestamp));
+
+    FString MessageTypeStr;
+    CHECK_PARSE(FJsonUtility::ParseString(Payload, UE::CPS::Properties::GType, MessageTypeStr));
+
+    FControlMessage::EType MessageType = FControlMessage::DeserializeType(MoveTemp(MessageTypeStr));
+
+    if (MessageType == EType::Invalid)
+    {
+        return FCaptureProtocolError(TEXT("Invalid message type"));
+    }
+
+    // Optionals
+    const TSharedPtr<FJsonObject>* Body;
+    TProtocolResult<void> ParseBodyResult = FJsonUtility::ParseObject(Payload, UE::CPS::Properties::GBody, Body);
+
+    const TSharedPtr<FJsonObject>* Error;
+    TProtocolResult<void> ParseErrorResult = FJsonUtility::ParseObject(Payload, UE::CPS::Properties::GError, Error);
+
+    FErrorResponse ErrorStruct;
+    if (ParseErrorResult.HasValue())
+    {
+        CHECK_PARSE(FJsonUtility::ParseString(*Error, UE::CPS::Properties::GName, ErrorStruct.Name));
+        CHECK_PARSE(FJsonUtility::ParseString(*Error, UE::CPS::Properties::GDescription, ErrorStruct.Description));
+    }
+
+    return FControlMessage(MoveTemp(SessionId),
+                           MoveTemp(AddressPath),
+                           TransactionId,
+                           Timestamp,
+                           MessageType,
+                           ParseBodyResult.HasValue() ? *Body : nullptr,
+                           MoveTemp(ErrorStruct));
+}
+
+TProtocolResult<FControlPacket> FControlMessage::Serialize(const FControlMessage& InMessage)
+{
+    TSharedPtr<FJsonObject> Payload = MakeShared<FJsonObject>();
+
+    Payload->SetStringField(UE::CPS::Properties::GSessionId, InMessage.SessionId);
+    Payload->SetStringField(UE::CPS::Properties::GAddressPath, InMessage.AddressPath);
+    Payload->SetNumberField(UE::CPS::Properties::GTransactionId, InMessage.TransactionId);
+    Payload->SetNumberField(UE::CPS::Properties::GTimestamp, InMessage.Timestamp);
+
+    FString MessageTypeStr = SerializeType(InMessage.MessageType);
+    Payload->SetStringField(UE::CPS::Properties::GType, MessageTypeStr);
+
+    if (InMessage.Body)
+    {
+        Payload->SetObjectField(UE::CPS::Properties::GBody, InMessage.Body);
+    }
+
+    if (!InMessage.Error.Name.IsEmpty())
+    {
+        TSharedPtr<FJsonObject> Error = MakeShared<FJsonObject>();
+        Error->SetStringField(UE::CPS::Properties::GName, InMessage.Error.Name);
+        Error->SetStringField(UE::CPS::Properties::GDescription, InMessage.Error.Description);
+
+        Payload->SetObjectField(UE::CPS::Properties::GError, Error);
+    }
+
+	TArray<uint8> Data;
+	if (!FJsonUtility::CreateUTF8DataFromJson(Payload, Data))
+	{
+		return FCaptureProtocolError(TEXT("Failed to serialize the payload"));
+	}
+
+    return FControlPacket(CPS_VERSION, MoveTemp(Data));
+}
+
+FControlMessage::FControlMessage(FString InSessionId,
+                                 FString InAddressPath,
+                                 uint32 InTransactionId,
+                                 uint64 InTimestamp,
+                                 EType InMessageType,
+                                 TSharedPtr<FJsonObject> InBody,
+                                 FErrorResponse InError)
+    : SessionId(MoveTemp(InSessionId))
+    , AddressPath(MoveTemp(InAddressPath))
+    , TransactionId(InTransactionId)
+    , Timestamp(InTimestamp)
+    , MessageType(InMessageType)
+    , Body(MoveTemp(InBody))
+    , Error(MoveTemp(InError))
+{
+}
+
+FControlMessage::FControlMessage(FString InAddressPath, 
+                                 EType InType,
+                                 TSharedPtr<FJsonObject> InBody)
+    : AddressPath(MoveTemp(InAddressPath))
+    , MessageType(InType)
+    , Body(MoveTemp(InBody))
+{
+}
+
+void FControlMessage::SetSessionId(FString InSessionId)
+{
+    SessionId = MoveTemp(InSessionId);
+}
+
+void FControlMessage::SetTransactionId(uint32 InTransactionId)
+{
+    TransactionId = InTransactionId;
+}
+
+void FControlMessage::SetTimestamp(uint64 InTimestamp)
+{
+    Timestamp = InTimestamp;
+}
+
+const FString& FControlMessage::GetSessionId() const
+{
+    return SessionId;
+}
+
+const FString& FControlMessage::GetAddressPath() const
+{
+    return AddressPath;
+}
+
+uint32 FControlMessage::GetTransactionId() const
+{
+    return TransactionId;
+}
+
+uint64 FControlMessage::GetTimestamp() const
+{
+    return Timestamp;
+}
+
+FControlMessage::EType FControlMessage::GetType() const
+{
+    return MessageType;
+}
+
+const TSharedPtr<FJsonObject>& FControlMessage::GetBody() const
+{
+    return Body;
+}
+
+TSharedPtr<FJsonObject>& FControlMessage::GetBody()
+{
+    return Body;
+}
+
+const FString& FControlMessage::GetErrorName() const
+{
+    return Error.Name;
+}
+
+const FString& FControlMessage::GetErrorDescription() const
+{
+    return Error.Description;
+}
+
+FControlMessage::EType FControlMessage::DeserializeType(FString InMessageTypeStr)
+{
+    if (InMessageTypeStr == UE::CPS::Properties::GRequest)
+    {
+        return EType::Request;
+    }
+    else if (InMessageTypeStr == UE::CPS::Properties::GResponse)
+    {
+        return EType::Response;
+    }
+    else if (InMessageTypeStr == UE::CPS::Properties::GUpdate)
+    {
+        return EType::Update;
+    }
+    else
+    {
+        return EType::Invalid;
+    }
+}
+
+FString FControlMessage::SerializeType(EType InMessageType)
+{
+    switch (InMessageType)
+    {
+        case EType::Request:
+            return UE::CPS::Properties::GRequest;
+        case EType::Response:
+            return UE::CPS::Properties::GResponse;
+        case EType::Update:
+            return UE::CPS::Properties::GUpdate;
+        default:
+            return TEXT("invalid");
+    }
+}
\ No newline at end of file
diff --git a/Engine/Plugins/VirtualProduction/CaptureManager/Source/CaptureProtocolStack/Private/Control/Messages/ControlRequest.cpp b/Engine/Plugins/VirtualProduction/CaptureManager/Source/CaptureProtocolStack/Private/Control/Messages/ControlRequest.cpp
new file mode 100644
--- /dev/null
+++ b/Engine/Plugins/VirtualProduction/CaptureManager/Source/CaptureProtocolStack/Private/Control/Messages/ControlRequest.cpp
@@ -0,0 +1,134 @@
+// Copyright Epic Games, Inc. All Rights Reserved.
+
+#include "Control/Messages/ControlRequest.h"
+#include "Control/Messages/Constants.h"
+
+FControlRequest::FControlRequest(FString InAddressPath)
+    : AddressPath(MoveTemp(InAddressPath))
+{
+}
+
+const FString& FControlRequest::GetAddressPath() const
+{
+    return AddressPath;
+}
+
+TSharedPtr<FJsonObject> FControlRequest::GetBody() const
+{
+    return nullptr;
+}
+
+FKeepAliveRequest::FKeepAliveRequest()
+    : FControlRequest(UE::CPS::AddressPaths::GKeepAlive)
+{
+}
+
+FStartSessionRequest::FStartSessionRequest()
+    : FControlRequest(UE::CPS::AddressPaths::GStartSession)
+{
+}
+
+FStopSessionRequest::FStopSessionRequest()
+    : FControlRequest(UE::CPS::AddressPaths::GStopSession)
+{
+}
+
+FGetServerInformationRequest::FGetServerInformationRequest()
+    : FControlRequest(UE::CPS::AddressPaths::GGetServerInformation)
+{
+}
+
+FSubscribeRequest::FSubscribeRequest()
+    : FControlRequest(UE::CPS::AddressPaths::GSubscribe)
+{
+}
+
+FUnsubscribeRequest::FUnsubscribeRequest()
+    : FControlRequest(UE::CPS::AddressPaths::GUnsubscribe)
+{
+}
+
+FGetStateRequest::FGetStateRequest()
+    : FControlRequest(UE::CPS::AddressPaths::GGetState)
+{
+}
+
+FStartRecordingTakeRequest::FStartRecordingTakeRequest(FString InSlateName, 
+                                                       uint16 InTakeNumber, 
+													   TOptional<FString> InSubject,
+													   TOptional<FString> InScenario,
+													   TOptional<TArray<FString>> InTags)
+    : FControlRequest(UE::CPS::AddressPaths::GStartRecordingTake)
+    , SlateName(MoveTemp(InSlateName))
+    , TakeNumber(InTakeNumber)
+    , Subject(MoveTemp(InSubject))
+    , Scenario(MoveTemp(InScenario))
+    , Tags(MoveTemp(InTags))
+{
+}
+
+TSharedPtr<FJsonObject> FStartRecordingTakeRequest::GetBody() const
+{
+    TSharedPtr<FJsonObject> Body = MakeShared<FJsonObject>();
+
+    Body->SetStringField(UE::CPS::Properties::GSlateName, SlateName);
+    Body->SetNumberField(UE::CPS::Properties::GTakeNumber, TakeNumber);
+	if (Subject.IsSet())
+	{
+		Body->SetStringField(UE::CPS::Properties::GSubject, Subject.GetValue());
+	}
+    
+	if (Scenario.IsSet())
+	{
+		Body->SetStringField(UE::CPS::Properties::GScenario, Scenario.GetValue());
+	}
+    
+	if (Tags.IsSet())
+	{
+		TArray<TSharedPtr<FJsonValue>> TagsJson;
+		for (const FString& Tag : Tags.GetValue())
+		{
+			TagsJson.Add(MakeShared<FJsonValueString>(Tag));
+		}
+
+		Body->SetArrayField(UE::CPS::Properties::GTags, TagsJson);
+	}
+    
+    return Body;
+}
+
+FStopRecordingTakeRequest::FStopRecordingTakeRequest()
+    : FControlRequest(UE::CPS::AddressPaths::GStopRecordingTake)
+{
+}
+
+FAbortRecordingTakeRequest::FAbortRecordingTakeRequest()
+	: FControlRequest(UE::CPS::AddressPaths::GAbortRecordingTake)
+{
+}
+
+FGetTakeListRequest::FGetTakeListRequest()
+    : FControlRequest(UE::CPS::AddressPaths::GGetTakeList)
+{
+}
+
+FGetTakeMetadataRequest::FGetTakeMetadataRequest(TArray<FString> InNames)
+    : FControlRequest(UE::CPS::AddressPaths::GGetTakeMetadata)
+    , Names(MoveTemp(InNames))
+{
+}
+
+TSharedPtr<FJsonObject> FGetTakeMetadataRequest::GetBody() const
+{
+    TSharedPtr<FJsonObject> Body = MakeShared<FJsonObject>();
+
+    TArray<TSharedPtr<FJsonValue>> NamesJson;
+    for (const FString& Name : Names)
+    {
+        NamesJson.Add(MakeShared<FJsonValueString>(Name));
+    }
+
+    Body->SetArrayField(UE::CPS::Properties::GNames, NamesJson);
+
+    return Body;
+}
diff --git a/Engine/Plugins/VirtualProduction/CaptureManager/Source/CaptureProtocolStack/Private/Control/Messages/ControlResponse.cpp b/Engine/Plugins/VirtualProduction/CaptureManager/Source/CaptureProtocolStack/Private/Control/Messages/ControlResponse.cpp
new file mode 100644
--- /dev/null
+++ b/Engine/Plugins/VirtualProduction/CaptureManager/Source/CaptureProtocolStack/Private/Control/Messages/ControlResponse.cpp
@@ -0,0 +1,294 @@
+// Copyright Epic Games, Inc. All Rights Reserved.
+
+#include "Control/Messages/ControlResponse.h"
+#include "Control/Messages/Constants.h"
+
+#include "Control/Messages/ControlJsonUtilities.h"
+
+FControlResponse::FControlResponse(FString InAddressPath)
+    : AddressPath(MoveTemp(InAddressPath))
+{
+}
+
+TProtocolResult<void> FControlResponse::Parse(TSharedPtr<FJsonObject> InBody)
+{
+    if (InBody.IsValid() && !InBody->Values.IsEmpty())
+    {
+        return FCaptureProtocolError("Response must NOT have a body");
+    }
+
+    return ResultOk;
+}
+
+const FString& FControlResponse::GetAddressPath() const
+{
+    return AddressPath;
+}
+
+FKeepAliveResponse::FKeepAliveResponse()
+    : FControlResponse(UE::CPS::AddressPaths::GKeepAlive)
+{
+}
+
+FStartSessionResponse::FStartSessionResponse()
+    : FControlResponse(UE::CPS::AddressPaths::GStartSession)
+{
+}
+
+TProtocolResult<void> FStartSessionResponse::Parse(TSharedPtr<FJsonObject> InBody)
+{
+    CHECK_PARSE(FJsonUtility::ParseString(InBody, UE::CPS::Properties::GSessionId, SessionId));
+
+    return ResultOk;
+}
+
+const FString& FStartSessionResponse::GetSessionId() const
+{
+    return SessionId;
+}
+
+FStopSessionResponse::FStopSessionResponse()
+    : FControlResponse(UE::CPS::AddressPaths::GStopSession)
+{
+}
+
+FGetServerInformationResponse::FGetServerInformationResponse()
+    : FControlResponse(UE::CPS::AddressPaths::GGetServerInformation)
+{
+}
+
+TProtocolResult<void> FGetServerInformationResponse::Parse(TSharedPtr<FJsonObject> InBody)
+{
+    CHECK_PARSE(FJsonUtility::ParseString(InBody, UE::CPS::Properties::GId, Id));
+    CHECK_PARSE(FJsonUtility::ParseString(InBody, UE::CPS::Properties::GName, Name));
+    CHECK_PARSE(FJsonUtility::ParseString(InBody, UE::CPS::Properties::GModel, Model));
+    CHECK_PARSE(FJsonUtility::ParseString(InBody, UE::CPS::Properties::GPlatformName, PlatformName));
+    CHECK_PARSE(FJsonUtility::ParseString(InBody, UE::CPS::Properties::GPlatformVersion, PlatformVersion));
+    CHECK_PARSE(FJsonUtility::ParseString(InBody, UE::CPS::Properties::GSoftwareName, SoftwareName));
+    CHECK_PARSE(FJsonUtility::ParseString(InBody, UE::CPS::Properties::GSoftwareVersion, SoftwareVersion));
+	CHECK_PARSE(FJsonUtility::ParseNumber(InBody, UE::CPS::Properties::GExportPort, ExportPort));
+
+    return ResultOk;
+}
+
+const FString& FGetServerInformationResponse::GetId() const
+{
+    return Id;
+}
+
+const FString& FGetServerInformationResponse::GetName() const
+{
+    return Name;
+}
+
+const FString& FGetServerInformationResponse::GetModel() const
+{
+    return Model;
+}
+
+const FString& FGetServerInformationResponse::GetPlatformName() const
+{
+    return PlatformName;
+}
+
+const FString& FGetServerInformationResponse::GetPlatformVersion() const
+{
+    return PlatformVersion;
+}
+
+const FString& FGetServerInformationResponse::GetSoftwareName() const
+{
+    return SoftwareName;
+}
+
+const FString& FGetServerInformationResponse::GetSoftwareVersion() const
+{
+    return SoftwareVersion;
+}
+
+uint16 FGetServerInformationResponse::GetExportPort() const
+{
+	return ExportPort;
+}
+
+FSubscribeResponse::FSubscribeResponse()
+    : FControlResponse(UE::CPS::AddressPaths::GSubscribe)
+{
+}
+
+FUnsubscribeResponse::FUnsubscribeResponse()
+    : FControlResponse(UE::CPS::AddressPaths::GUnsubscribe)
+{
+}
+
+FGetStateResponse::FGetStateResponse()
+    : FControlResponse(UE::CPS::AddressPaths::GGetState)
+{
+}
+
+TProtocolResult<void> FGetStateResponse::Parse(TSharedPtr<FJsonObject> InBody)
+{
+    CHECK_PARSE(FJsonUtility::ParseBool(InBody, UE::CPS::Properties::GIsRecording, bIsRecording));
+
+    // Optional
+    const TSharedPtr<FJsonObject>* PlatformStatePtr;
+    TProtocolResult<void> Result = FJsonUtility::ParseObject(InBody, UE::CPS::Properties::GPlatformState, PlatformStatePtr);
+	if (Result.HasValue())
+	{
+		PlatformState = *PlatformStatePtr;
+	}
+
+    return ResultOk;
+}
+
+bool FGetStateResponse::IsRecording() const
+{
+    return bIsRecording;
+}
+
+const TSharedPtr<FJsonObject>& FGetStateResponse::GetPlatformState() const
+{
+    return PlatformState;
+}
+
+FStartRecordingTakeResponse::FStartRecordingTakeResponse()
+    : FControlResponse(UE::CPS::AddressPaths::GStartRecordingTake)
+{
+}
+
+FStopRecordingTakeResponse::FStopRecordingTakeResponse()
+    : FControlResponse(UE::CPS::AddressPaths::GStopRecordingTake)
+{
+}
+
+FAbortRecordingTakeResponse::FAbortRecordingTakeResponse()
+	: FControlResponse(UE::CPS::AddressPaths::GAbortRecordingTake)
+{
+}
+
+TProtocolResult<void> FStopRecordingTakeResponse::Parse(TSharedPtr<FJsonObject> InBody)
+{
+	CHECK_PARSE(FJsonUtility::ParseString(InBody, UE::CPS::Properties::GName, TakeName));
+
+	return ResultOk;
+}
+
+const FString& FStopRecordingTakeResponse::GetTakeName() const
+{
+	return TakeName;
+}
+
+FGetTakeListResponse::FGetTakeListResponse()
+    : FControlResponse(UE::CPS::AddressPaths::GGetTakeList)
+{
+}
+
+TProtocolResult<void> FGetTakeListResponse::Parse(TSharedPtr<FJsonObject> InBody)
+{
+    const TArray<TSharedPtr<FJsonValue>>* NamesJson;
+    CHECK_PARSE(FJsonUtility::ParseArray(InBody, UE::CPS::Properties::GNames, NamesJson));
+
+    for (const TSharedPtr<FJsonValue>& NameJson : *NamesJson)
+    {
+        Names.Add(NameJson->AsString());
+    }
+
+    return ResultOk;
+}
+
+const TArray<FString>& FGetTakeListResponse::GetNames() const
+{
+    return Names;
+}
+
+FGetTakeMetadataResponse::FGetTakeMetadataResponse()
+    : FControlResponse(UE::CPS::AddressPaths::GGetTakeMetadata)
+{
+}
+
+TProtocolResult<void> FGetTakeMetadataResponse::Parse(TSharedPtr<FJsonObject> InBody)
+{
+    const TArray<TSharedPtr<FJsonValue>>* TakesJson;
+    CHECK_PARSE(FJsonUtility::ParseArray(InBody, UE::CPS::Properties::GTakes, TakesJson));
+
+    for (const TSharedPtr<FJsonValue>& TakeJson : *TakesJson)
+    {
+        const TSharedPtr<FJsonObject>& TakeJsonObject = TakeJson->AsObject();
+
+        FTakeObject Take;
+        CHECK_PARSE(CreateTakeObject(TakeJsonObject, Take));
+        Takes.Add(MoveTemp(Take));
+    }
+
+    return ResultOk;
+}
+
+const TArray<FGetTakeMetadataResponse::FTakeObject>& FGetTakeMetadataResponse::GetTakes() const
+{
+    return Takes;
+}
+
+TProtocolResult<void> FGetTakeMetadataResponse::CreateTakeObject(const TSharedPtr<FJsonObject>& InTakeObject, FTakeObject& OutTake) const
+{
+    CHECK_PARSE(FJsonUtility::ParseString(InTakeObject, UE::CPS::Properties::GName, OutTake.Name));
+    CHECK_PARSE(FJsonUtility::ParseString(InTakeObject, UE::CPS::Properties::GSlateName, OutTake.Slate));
+    CHECK_PARSE(FJsonUtility::ParseNumber(InTakeObject, UE::CPS::Properties::GTakeNumber, OutTake.TakeNumber));
+    CHECK_PARSE(FJsonUtility::ParseString(InTakeObject, UE::CPS::Properties::GDateTime, OutTake.DateTime));
+    CHECK_PARSE(FJsonUtility::ParseString(InTakeObject, UE::CPS::Properties::GAppVersion, OutTake.AppVersion));
+    CHECK_PARSE(FJsonUtility::ParseString(InTakeObject, UE::CPS::Properties::GModel, OutTake.Model));
+
+    // Optional
+    FJsonUtility::ParseString(InTakeObject, UE::CPS::Properties::GSubject, OutTake.Subject);
+    FJsonUtility::ParseString(InTakeObject, UE::CPS::Properties::GScenario, OutTake.Scenario);
+
+    const TArray<TSharedPtr<FJsonValue>>* TagsJson;
+    TProtocolResult<void> ParseArrayResult = FJsonUtility::ParseArray(InTakeObject, UE::CPS::Properties::GTags, TagsJson);
+
+    if (ParseArrayResult.HasValue())
+    {
+        for (const TSharedPtr<FJsonValue>& TagJson : *TagsJson)
+        {
+            OutTake.Tags.Add(TagJson->AsString());
+        }
+    }
+
+    const TArray<TSharedPtr<FJsonValue>>* FilesJson;
+    CHECK_PARSE(FJsonUtility::ParseArray(InTakeObject, UE::CPS::Properties::GFiles, FilesJson));
+
+    for (const TSharedPtr<FJsonValue>& FileJson : *FilesJson)
+    {
+        const TSharedPtr<FJsonObject>& FileJsonObject = FileJson->AsObject();
+        FFileObject File;
+        CHECK_PARSE(CreateFileObject(FileJsonObject, File));
+        OutTake.Files.Add(MoveTemp(File));
+    }
+
+    const TSharedPtr<FJsonObject>* VideoMetadata;
+    TProtocolResult<void> ParseVideoObject = FJsonUtility::ParseObject(InTakeObject, UE::CPS::Properties::GVideo, VideoMetadata);
+    if (ParseVideoObject.HasValue())
+    {
+		CHECK_PARSE(FJsonUtility::ParseNumber(*VideoMetadata, UE::CPS::Properties::GFrames, OutTake.Video.Frames));
+        CHECK_PARSE(FJsonUtility::ParseNumber(*VideoMetadata, UE::CPS::Properties::GFrameRate, OutTake.Video.FrameRate));
+        CHECK_PARSE(FJsonUtility::ParseNumber(*VideoMetadata, UE::CPS::Properties::GHeight, OutTake.Video.Height));
+        CHECK_PARSE(FJsonUtility::ParseNumber(*VideoMetadata, UE::CPS::Properties::GWidth, OutTake.Video.Width));
+    }
+
+    const TSharedPtr<FJsonObject>* AudioMetadata;
+    TProtocolResult<void> ParseAudioObject = FJsonUtility::ParseObject(InTakeObject, UE::CPS::Properties::GAudio, AudioMetadata);
+    if (ParseAudioObject.HasValue())
+    {
+        CHECK_PARSE(FJsonUtility::ParseNumber(*AudioMetadata, UE::CPS::Properties::GChannels, OutTake.Audio.Channels));
+        CHECK_PARSE(FJsonUtility::ParseNumber(*AudioMetadata, UE::CPS::Properties::GSampleRate, OutTake.Audio.SampleRate));
+        CHECK_PARSE(FJsonUtility::ParseNumber(*AudioMetadata, UE::CPS::Properties::GBitsPerChannel, OutTake.Audio.BitsPerChannel));
+    }
+
+    return ResultOk;
+}
+
+TProtocolResult<void> FGetTakeMetadataResponse::CreateFileObject(const TSharedPtr<FJsonObject>& InFileObject, FFileObject& OutFile) const
+{
+    CHECK_PARSE(FJsonUtility::ParseString(InFileObject, UE::CPS::Properties::GName, OutFile.Name));
+    CHECK_PARSE(FJsonUtility::ParseNumber(InFileObject, UE::CPS::Properties::GLength, OutFile.Length));
+
+    return ResultOk;
+}
diff --git a/Engine/Plugins/VirtualProduction/CaptureManager/Source/CaptureProtocolStack/Private/Control/Messages/ControlUpdate.cpp b/Engine/Plugins/VirtualProduction/CaptureManager/Source/CaptureProtocolStack/Private/Control/Messages/ControlUpdate.cpp
new file mode 100644
--- /dev/null
+++ b/Engine/Plugins/VirtualProduction/CaptureManager/Source/CaptureProtocolStack/Private/Control/Messages/ControlUpdate.cpp
@@ -0,0 +1,203 @@
+// Copyright Epic Games, Inc. All Rights Reserved.
+
+#include "Control/Messages/ControlUpdate.h"
+#include "Control/Messages/Constants.h"
+
+#include "Control/Messages/ControlJsonUtilities.h"
+
+TProtocolResult<TSharedRef<FControlUpdate>> FControlUpdateCreator::Create(const FString& InAddressPath)
+{
+    if (InAddressPath == UE::CPS::AddressPaths::GSessionStopped)
+    {
+        return TProtocolResult<TSharedRef<FControlUpdate>>(MakeShared<FSessionStopped>());
+    }
+	else if (InAddressPath == UE::CPS::AddressPaths::GTakeAdded)
+	{
+		return TProtocolResult<TSharedRef<FControlUpdate>>(MakeShared<FTakeAddedUpdate>());
+	}
+	else if (InAddressPath == UE::CPS::AddressPaths::GTakeRemoved)
+	{
+		return TProtocolResult<TSharedRef<FControlUpdate>>(MakeShared<FTakeRemovedUpdate>());
+	}
+	else if (InAddressPath == UE::CPS::AddressPaths::GTakeUpdated)
+	{
+		return TProtocolResult<TSharedRef<FControlUpdate>>(MakeShared<FTakeUpdatedUpdate>());
+	}
+    else if (InAddressPath == UE::CPS::AddressPaths::GRecordingStatus)
+    {
+        return TProtocolResult<TSharedRef<FControlUpdate>>(MakeShared<FRecordingStatusUpdate>());
+    }
+    else if(InAddressPath == UE::CPS::AddressPaths::GDiskCapacity)
+    {
+        return TProtocolResult<TSharedRef<FControlUpdate>>(MakeShared<FDiskCapacityUpdate>());
+    }
+    else if (InAddressPath == UE::CPS::AddressPaths::GBattery)
+    {
+        return TProtocolResult<TSharedRef<FControlUpdate>>(MakeShared<FBatteryPercentageUpdate>());
+    }
+    else if (InAddressPath == UE::CPS::AddressPaths::GThermalState)
+    {
+        return TProtocolResult<TSharedRef<FControlUpdate>>(MakeShared<FThermalStateUpdate>());
+    }
+    else
+    {
+        return FCaptureProtocolError(TEXT("Unknown update arrived"));
+    }
+}
+
+FControlUpdate::FControlUpdate(FString InAddressPath)
+    : AddressPath(MoveTemp(InAddressPath))
+{
+}
+
+const FString& FControlUpdate::GetAddressPath() const
+{
+    return AddressPath;
+}
+
+TProtocolResult<void> FControlUpdate::Parse(TSharedPtr<FJsonObject> InBody)
+{
+    if (InBody.IsValid() && !InBody->Values.IsEmpty())
+    {
+        return FCaptureProtocolError("Update must NOT have a body");
+    }
+
+    return ResultOk;
+}
+
+FSessionStopped::FSessionStopped()
+    : FControlUpdate(UE::CPS::AddressPaths::GSessionStopped)
+{
+}
+
+FRecordingStatusUpdate::FRecordingStatusUpdate()
+    : FControlUpdate(UE::CPS::AddressPaths::GRecordingStatus)
+{
+}
+
+TProtocolResult<void> FRecordingStatusUpdate::Parse(TSharedPtr<FJsonObject> InBody)
+{
+    CHECK_PARSE(FJsonUtility::ParseBool(InBody, UE::CPS::Properties::GIsRecording, bIsRecording));
+
+    return ResultOk;
+}
+
+bool FRecordingStatusUpdate::IsRecording() const
+{
+    return bIsRecording;
+}
+
+TProtocolResult<void> FBaseTakeUpdate::Parse(TSharedPtr<FJsonObject> InBody)
+{
+	CHECK_PARSE(FJsonUtility::ParseString(InBody, UE::CPS::Properties::GName, TakeName));
+
+	return ResultOk;
+}
+
+const FString& FBaseTakeUpdate::GetTakeName() const
+{
+	return TakeName;
+}
+
+FTakeAddedUpdate::FTakeAddedUpdate()
+	: FBaseTakeUpdate(UE::CPS::AddressPaths::GTakeAdded)
+{
+}
+
+FTakeRemovedUpdate::FTakeRemovedUpdate()
+	: FBaseTakeUpdate(UE::CPS::AddressPaths::GTakeRemoved)
+{
+}
+
+FTakeUpdatedUpdate::FTakeUpdatedUpdate()
+	: FBaseTakeUpdate(UE::CPS::AddressPaths::GTakeUpdated)
+{
+}
+
+FDiskCapacityUpdate::FDiskCapacityUpdate()
+    : FControlUpdate(UE::CPS::AddressPaths::GDiskCapacity)
+{
+}
+
+TProtocolResult<void> FDiskCapacityUpdate::Parse(TSharedPtr<FJsonObject> InBody)
+{
+    CHECK_PARSE(FJsonUtility::ParseNumber(InBody, UE::CPS::Properties::GTotal, Total));
+    CHECK_PARSE(FJsonUtility::ParseNumber(InBody, UE::CPS::Properties::GRemaining, Remaining));
+
+    return ResultOk;
+}
+
+uint64 FDiskCapacityUpdate::GetTotal() const
+{
+    return Total;
+}
+
+uint64 FDiskCapacityUpdate::GetRemaining() const
+{
+    return Remaining;
+}
+
+FBatteryPercentageUpdate::FBatteryPercentageUpdate()
+    : FControlUpdate(UE::CPS::AddressPaths::GBattery)
+{
+}
+
+TProtocolResult<void> FBatteryPercentageUpdate::Parse(TSharedPtr<FJsonObject> InBody)
+{
+    CHECK_PARSE(FJsonUtility::ParseNumber(InBody, UE::CPS::Properties::GLevel, Level));
+
+    return ResultOk;
+}
+
+float FBatteryPercentageUpdate::GetLevel() const
+{
+    return Level;
+}
+
+FThermalStateUpdate::FThermalStateUpdate()
+    : FControlUpdate(UE::CPS::AddressPaths::GThermalState)
+{
+}
+
+TProtocolResult<void> FThermalStateUpdate::Parse(TSharedPtr<FJsonObject> InBody)
+{
+    FString StateStr;
+    CHECK_PARSE(FJsonUtility::ParseString(InBody, UE::CPS::Properties::GState, StateStr));
+
+    State = ConvertState(StateStr);
+    if (State == EState::Invalid)
+    {
+        return FCaptureProtocolError("Invalid thermal state provided: " + StateStr);
+    }
+
+    return ResultOk;
+}
+
+FThermalStateUpdate::EState FThermalStateUpdate::GetState() const
+{
+    return State;
+}
+
+FThermalStateUpdate::EState FThermalStateUpdate::ConvertState(const FString& InStateString)
+{
+    if (InStateString == UE::CPS::Properties::GNominal)
+    {
+        return EState::Nominal;
+    }
+    else if (InStateString == UE::CPS::Properties::GFair)
+    {
+        return EState::Fair;
+    }
+    else if (InStateString == UE::CPS::Properties::GSerious)
+    {
+        return EState::Serious;
+    }
+    else if (InStateString == UE::CPS::Properties::GCritical)
+    {
+        return EState::Critical;
+    }
+    else
+    {
+        return EState::Invalid;
+    }
+}
diff --git a/Engine/Plugins/VirtualProduction/CaptureManager/Source/CaptureProtocolStack/Private/Discovery/Communication/DiscoveryCommunication.cpp b/Engine/Plugins/VirtualProduction/CaptureManager/Source/CaptureProtocolStack/Private/Discovery/Communication/DiscoveryCommunication.cpp
new file mode 100644
--- /dev/null
+++ b/Engine/Plugins/VirtualProduction/CaptureManager/Source/CaptureProtocolStack/Private/Discovery/Communication/DiscoveryCommunication.cpp
@@ -0,0 +1,68 @@
+// Copyright Epic Games, Inc. All Rights Reserved.
+
+#include "Discovery/Communication/DiscoveryCommunication.h"
+
+DEFINE_LOG_CATEGORY(LogCPSDiscoveryCommunication)
+
+const uint16 FDiscoveryCommunication::MulticastPort = 27838; // According to the specification
+const FString FDiscoveryCommunication::MulticastAddress = TEXT("239.255.137.139"); // According to the specification
+FDiscoveryCommunication::FDiscoveryCommunication()
+    : SynchronizedReceiver(FRunnerType::FOnProcess::CreateRaw(this, &FDiscoveryCommunication::OnProcessReceivedPacket))
+{
+}
+
+TProtocolResult<void> FDiscoveryCommunication::Start()
+{
+    CPS_CHECK_VOID_RESULT(Client.Init({ MulticastPort, MulticastAddress }, FOnSocketDataReceived::CreateRaw(this, &FDiscoveryCommunication::OnPacketArrived)));
+    CPS_CHECK_VOID_RESULT(Client.Start());
+
+    return ResultOk;
+}
+
+TProtocolResult<void> FDiscoveryCommunication::Stop()
+{
+    return Client.Stop();
+}
+
+void FDiscoveryCommunication::SendMessage(FDiscoveryPacket InMessage, const FString& InEndpoint)
+{
+    TProtocolResult<TArray<uint8>> SerializeResult = FDiscoveryPacket::Serialize(InMessage);
+
+    if (SerializeResult.HasError())
+    {
+        UE_LOG(LogCPSDiscoveryCommunication, Error, TEXT("Invalid message: %s"), *(SerializeResult.StealError().GetMessage()));
+        return;
+    }
+
+    Client.SendMessage(SerializeResult.StealValue(), InEndpoint);
+}
+
+void FDiscoveryCommunication::SendMessage(FDiscoveryPacket InMessage, const FString& InEndpointIp, const uint16 InEndpointPort)
+{
+    SendMessage(MoveTemp(InMessage), InEndpointIp + TEXT(":") + FString::FromInt(InEndpointPort));
+}
+
+void FDiscoveryCommunication::SetReceiveHandler(FOnPacketReceived InReceiveHandler)
+{
+    OnPacketReceived = MoveTemp(InReceiveHandler);
+}
+
+void FDiscoveryCommunication::OnPacketArrived(const FArrayReaderPtr& InPayload, const FIPv4Endpoint& InEndpoint)
+{
+    const TArray<uint8>& Payload = *InPayload;
+
+    TProtocolResult<FDiscoveryPacket> DeserializeResult = FDiscoveryPacket::Deserialize(Payload);
+    if (DeserializeResult.HasError())
+    {
+        UE_LOG(LogCPSDiscoveryCommunication, Error, TEXT("Invalid message: %s"), *(DeserializeResult.StealError().GetMessage()));
+        return;
+    }
+
+    FDiscoveryPacket Packet = DeserializeResult.StealValue();
+    SynchronizedReceiver.Add(MoveTemp(Packet));
+}
+
+void FDiscoveryCommunication::OnProcessReceivedPacket(FDiscoveryPacket InPacket)
+{
+    OnPacketReceived.ExecuteIfBound(MoveTemp(InPacket));
+}
\ No newline at end of file
diff --git a/Engine/Plugins/VirtualProduction/CaptureManager/Source/CaptureProtocolStack/Private/Discovery/Communication/DiscoveryPacket.cpp b/Engine/Plugins/VirtualProduction/CaptureManager/Source/CaptureProtocolStack/Private/Discovery/Communication/DiscoveryPacket.cpp
new file mode 100644
--- /dev/null
+++ b/Engine/Plugins/VirtualProduction/CaptureManager/Source/CaptureProtocolStack/Private/Discovery/Communication/DiscoveryPacket.cpp
@@ -0,0 +1,97 @@
+// Copyright Epic Games, Inc. All Rights Reserved.
+
+#include "Discovery/Communication/DiscoveryPacket.h"
+
+#include "EngineLogs.h"
+
+const TArray<uint8> FDiscoveryPacket::Header = { 'C', 'P', 'S', 'D', 'I', 'S', 'C', 'O', 'V', 'E', 'R', '\0' };
+
+FDiscoveryPacket::FDiscoveryPacket(EMessageType InMessageType, TArray<uint8> InPayload)
+    : MessageType(InMessageType)
+	, Payload(MoveTemp(InPayload))
+{
+}
+
+TProtocolResult<FDiscoveryPacket> FDiscoveryPacket::Deserialize(const TConstArrayView<uint8>& InData)
+{
+    if (InData.Num() < Header.Num() + 1)
+    {
+        return FCaptureProtocolError(TEXT("Message with incorrect number of bytes arrived"));
+    }
+
+    int32 Result = FMemory::Memcmp(InData.GetData(), Header.GetData(), Header.Num());
+    if (Result != 0)
+    {
+        return FCaptureProtocolError(TEXT("Message with incorrect header arrived"));
+    }
+
+	EMessageType MessageType = ToMessageType(InData[Header.Num()]);
+
+	if (MessageType == EMessageType::Invalid)
+	{
+		return FCaptureProtocolError(TEXT("Invalid message type field"));
+	}
+
+    FDiscoveryPacket Packet;
+
+	Packet.MessageType = MessageType;
+    Packet.Payload = InData.RightChop(Header.Num() + sizeof(MessageType));
+
+    return Packet;
+}
+
+TProtocolResult<TArray<uint8>> FDiscoveryPacket::Serialize(const FDiscoveryPacket& InMessage)
+{
+    TArray<uint8> OutData;
+
+    OutData.Append(InMessage.Header.GetData(), InMessage.Header.Num());
+	OutData.Add(FromMessageType(InMessage.MessageType));
+    OutData.Append(InMessage.Payload.GetData(), InMessage.Payload.Num());
+
+    return OutData;
+}
+
+const FDiscoveryPacket::EMessageType FDiscoveryPacket::GetMessageType() const
+{
+	return MessageType;
+}
+
+const TArray<uint8>& FDiscoveryPacket::GetPayload() const
+{
+    return Payload;
+}
+
+FDiscoveryPacket::EMessageType FDiscoveryPacket::ToMessageType(uint8 InMessageType)
+{
+	if (InMessageType == static_cast<uint8>(EMessageType::Request))
+	{
+		return EMessageType::Request;
+	}
+	else if (InMessageType == static_cast<uint8>(EMessageType::Response))
+	{
+		return EMessageType::Response;
+	}
+	else if (InMessageType == static_cast<uint8>(EMessageType::Notify))
+	{
+		return EMessageType::Notify;
+	}
+	else
+	{
+		return EMessageType::Invalid;
+	}
+}
+
+uint8 FDiscoveryPacket::FromMessageType(EMessageType InMessageType)
+{
+	switch (InMessageType)
+	{
+		case EMessageType::Request:
+			return 0;
+		case EMessageType::Response:
+			return 1;
+		case EMessageType::Notify:
+			return 2;
+		default:
+			return MAX_uint8;
+	}
+}
\ No newline at end of file
diff --git a/Engine/Plugins/VirtualProduction/CaptureManager/Source/CaptureProtocolStack/Private/Discovery/DiscoveryMessenger.cpp b/Engine/Plugins/VirtualProduction/CaptureManager/Source/CaptureProtocolStack/Private/Discovery/DiscoveryMessenger.cpp
new file mode 100644
--- /dev/null
+++ b/Engine/Plugins/VirtualProduction/CaptureManager/Source/CaptureProtocolStack/Private/Discovery/DiscoveryMessenger.cpp
@@ -0,0 +1,105 @@
+// Copyright Epic Games, Inc. All Rights Reserved.
+
+#include "Discovery/DiscoveryMessenger.h"
+
+DEFINE_LOG_CATEGORY(LogCPSDiscoveryMessenger)
+
+FDiscoveryMessenger::FDiscoveryMessenger()
+{
+}
+
+FDiscoveryMessenger::~FDiscoveryMessenger()
+{
+    Stop();
+}
+
+TProtocolResult<void> FDiscoveryMessenger::Start()
+{
+    CPS_CHECK_VOID_RESULT(Communication.Start());
+
+    Communication.SetReceiveHandler(
+        FDiscoveryCommunication::FOnPacketReceived::CreateRaw(this, &FDiscoveryMessenger::OnPacketArrived));
+
+    return ResultOk;
+}
+
+TProtocolResult<void> FDiscoveryMessenger::Stop()
+{
+    return Communication.Stop();
+}
+
+TProtocolResult<void> FDiscoveryMessenger::SendMulticastRequest()
+{
+    constexpr int32 NumberOfRequests = 3;
+
+    FDiscoveryRequest Request;
+    TProtocolResult<FDiscoveryPacket> SerializeResult = FDiscoveryRequest::Serialize(Request);
+
+    if (SerializeResult.HasError())
+    {
+        return FCaptureProtocolError(TEXT("Failed to serialize request"));
+    }
+
+    for (int32 Counter = 0; Counter < NumberOfRequests; ++Counter)
+    {
+		FDiscoveryPacket Packet = SerializeResult.GetValue();
+        Communication.SendMessage(MoveTemp(Packet), FDiscoveryCommunication::MulticastAddress, FDiscoveryCommunication::MulticastPort);
+    }
+
+    return ResultOk;
+}
+
+void FDiscoveryMessenger::SetResponseHandler(FOnResponseArrived InOnResponse)
+{
+    OnResponse = MoveTemp(InOnResponse);
+}
+
+void FDiscoveryMessenger::SetNotifyHandler(FOnNotifyArrived InOnNotify)
+{
+    OnNotify = MoveTemp(InOnNotify);
+}
+
+void FDiscoveryMessenger::OnPacketArrived(FDiscoveryPacket InPacket)
+{
+    const TArray<uint8>& Payload = InPacket.GetPayload();
+
+    if (InPacket.GetMessageType() == FDiscoveryPacket::EMessageType::Request)
+    {
+        UE_LOG(LogCPSDiscoveryMessenger, Error, TEXT("Client currently doesn't support requests."));
+        return;
+    }
+    else if (InPacket.GetMessageType() == FDiscoveryPacket::EMessageType::Response)
+    {
+        TProtocolResult<FDiscoveryResponse> DeserializeResult = FDiscoveryResponse::Deserialize(InPacket);
+
+        if (DeserializeResult.HasError())
+        {
+            UE_LOG(LogCPSDiscoveryMessenger, Error, TEXT("Failed to parse the response message."));
+            return;
+        }
+
+        FDiscoveryResponse Response = DeserializeResult.StealValue();
+
+        OnResponse.ExecuteIfBound(MoveTemp(Response));
+    }
+    else if (InPacket.GetMessageType() == FDiscoveryPacket::EMessageType::Notify)
+    {
+        TProtocolResult<FDiscoveryNotify> DeserializeResult = FDiscoveryNotify::Deserialize(InPacket);
+
+        if (DeserializeResult.HasError())
+        {
+            UE_LOG(LogCPSDiscoveryMessenger, Error, TEXT("Failed to parse the notify message."));
+            return;
+        }
+
+        FDiscoveryNotify Notify = DeserializeResult.StealValue();
+
+        OnNotify.ExecuteIfBound(MoveTemp(Notify));
+    }
+    else
+    {
+        UE_LOG(LogCPSDiscoveryMessenger, Error, TEXT("Invalid message arrived."));
+        return;
+    }
+}
+
diff --git a/Engine/Plugins/VirtualProduction/CaptureManager/Source/CaptureProtocolStack/Private/Discovery/Messages/DiscoveryNotify.cpp b/Engine/Plugins/VirtualProduction/CaptureManager/Source/CaptureProtocolStack/Private/Discovery/Messages/DiscoveryNotify.cpp
new file mode 100644
--- /dev/null
+++ b/Engine/Plugins/VirtualProduction/CaptureManager/Source/CaptureProtocolStack/Private/Discovery/Messages/DiscoveryNotify.cpp
@@ -0,0 +1,118 @@
+// Copyright Epic Games, Inc. All Rights Reserved.
+
+#include "Discovery/Messages/DiscoveryNotify.h"
+#include "Containers/ArrayView.h"
+
+const uint32 FDiscoveryNotify::MinPayloadSize = 16 + 2 + 1 + 2; // ServerId + ControlPort + ConnectionState + Supported Versions (at least 1)
+FDiscoveryNotify::FDiscoveryNotify(FServerId InServerId, uint16 InControlPort, EConnectionState InConnectionState, TArray<uint16> InSupportedVersions)
+    : ServerId(MoveTemp(InServerId))
+    , ControlPort(InControlPort)
+    , ConnectionState(InConnectionState)
+	, SupportedVersions(MoveTemp(InSupportedVersions))
+{
+}
+
+TProtocolResult<FDiscoveryNotify> FDiscoveryNotify::Deserialize(const FDiscoveryPacket& InPacket)
+{
+	if (InPacket.GetMessageType() != FDiscoveryPacket::EMessageType::Notify)	
+	{ 
+		return FCaptureProtocolError(TEXT("Invalid request arrived"));
+	}
+
+	TConstArrayView<uint8> Payload = InPacket.GetPayload();
+
+	if (Payload.Num() < MinPayloadSize)
+    {
+        return FCaptureProtocolError(TEXT("Invalid number of bytes"));
+    }
+
+    FServerId ServerId;
+    FMemory::Memcpy(ServerId.GetData(), Payload.GetData(), ServerId.Num());
+
+    uint16 ControlPort = *reinterpret_cast<const uint16*>(Payload.GetData() + ServerId.Num());
+    uint8 ConnectionState = Payload[ServerId.Num() + sizeof(ControlPort)];
+
+    EConnectionState ConnectionStateEnum = ToConnectionState(ConnectionState);
+
+    if (ConnectionStateEnum == EConnectionState::Invalid)
+    {
+        return FCaptureProtocolError(TEXT("Invalid connection state field"));
+    }
+
+	TConstArrayView<uint8> SupportedVersionsBytes = Payload.RightChop(ServerId.Num() + sizeof(ControlPort) + sizeof(ConnectionStateEnum));
+	TArray<uint16> SupportedVersions;
+	SupportedVersions.Append((uint16*)SupportedVersionsBytes.GetData(), SupportedVersionsBytes.Num() / 2);
+
+    return FDiscoveryNotify(MoveTemp(ServerId), ControlPort, ConnectionStateEnum, MoveTemp(SupportedVersions));
+}
+
+TProtocolResult<FDiscoveryPacket> FDiscoveryNotify::Serialize(const FDiscoveryNotify& InNotify)
+{
+    TArray<uint8> Payload;
+    Payload.Reserve(InNotify.ServerId.Num() + sizeof(InNotify.ControlPort) + sizeof(InNotify.ConnectionState) + InNotify.SupportedVersions.Num() * sizeof(uint16));
+
+    Payload.Append(InNotify.ServerId.GetData(), InNotify.ServerId.Num());
+    Payload.Append((uint8*) &InNotify.ControlPort, sizeof(InNotify.ControlPort));
+
+    uint8 ConnectionState = FromConnectionState(InNotify.ConnectionState);
+
+    if (ConnectionState == MAX_uint8)
+    {
+        return FCaptureProtocolError(TEXT("Invalid connection state field"));
+    }
+
+    Payload.Add(ConnectionState);
+	Payload.Append((uint8*) InNotify.SupportedVersions.GetData(), InNotify.SupportedVersions.Num() * sizeof(uint16));
+
+    return FDiscoveryPacket(FDiscoveryPacket::EMessageType::Notify, MoveTemp(Payload));
+}
+
+const FDiscoveryNotify::FServerId& FDiscoveryNotify::GetServerId() const
+{
+    return ServerId;
+}
+
+uint16 FDiscoveryNotify::GetControlPort() const
+{
+    return ControlPort;
+}
+
+FDiscoveryNotify::EConnectionState FDiscoveryNotify::GetConnectionState() const
+{
+    return ConnectionState;
+}
+
+const TArray<uint16>& FDiscoveryNotify::GetSupportedVersions() const
+{
+	return SupportedVersions;
+}
+
+FDiscoveryNotify::EConnectionState FDiscoveryNotify::ToConnectionState(uint8 InConnectionState)
+{
+    if (InConnectionState == 0)
+    {
+        return EConnectionState::Offline;
+    }
+    else if (InConnectionState == 1)
+    {
+        return EConnectionState::Online;
+    }
+    else
+    {
+        return EConnectionState::Invalid;
+    }
+}
+
+uint8 FDiscoveryNotify::FromConnectionState(EConnectionState InConnectionState)
+{
+    switch (InConnectionState)
+    {
+        case EConnectionState::Offline:
+            return 0;
+        case EConnectionState::Online:
+            return 1;
+        case EConnectionState::Invalid:
+        default:
+            return MAX_uint8;
+    }
+}
\ No newline at end of file
diff --git a/Engine/Plugins/VirtualProduction/CaptureManager/Source/CaptureProtocolStack/Private/Discovery/Messages/DiscoveryRequest.cpp b/Engine/Plugins/VirtualProduction/CaptureManager/Source/CaptureProtocolStack/Private/Discovery/Messages/DiscoveryRequest.cpp
new file mode 100644
--- /dev/null
+++ b/Engine/Plugins/VirtualProduction/CaptureManager/Source/CaptureProtocolStack/Private/Discovery/Messages/DiscoveryRequest.cpp
@@ -0,0 +1,26 @@
+// Copyright Epic Games, Inc. All Rights Reserved.
+
+#include "Discovery/Messages/DiscoveryRequest.h"
+#include "Containers/ArrayView.h"
+
+const uint32 FDiscoveryRequest::PayloadSize = 0;
+
+TProtocolResult<FDiscoveryRequest> FDiscoveryRequest::Deserialize(const FDiscoveryPacket& InPacket)
+{
+	if (InPacket.GetMessageType() != FDiscoveryPacket::EMessageType::Request)
+	{
+		return FCaptureProtocolError(TEXT("Invalid request arrived"));
+	}
+
+    if (InPacket.GetPayload().Num() != PayloadSize)
+    {
+        return FCaptureProtocolError(TEXT("Invalid number of bytes"));
+    }
+
+    return FDiscoveryRequest();
+}
+
+TProtocolResult<FDiscoveryPacket> FDiscoveryRequest::Serialize(const FDiscoveryRequest& InRequest)
+{
+    return FDiscoveryPacket(FDiscoveryPacket::EMessageType::Request, TArray<uint8>());
+}
\ No newline at end of file
diff --git a/Engine/Plugins/VirtualProduction/CaptureManager/Source/CaptureProtocolStack/Private/Discovery/Messages/DiscoveryResponse.cpp b/Engine/Plugins/VirtualProduction/CaptureManager/Source/CaptureProtocolStack/Private/Discovery/Messages/DiscoveryResponse.cpp
new file mode 100644
--- /dev/null
+++ b/Engine/Plugins/VirtualProduction/CaptureManager/Source/CaptureProtocolStack/Private/Discovery/Messages/DiscoveryResponse.cpp
@@ -0,0 +1,65 @@
+// Copyright Epic Games, Inc. All Rights Reserved.
+
+#include "Discovery/Messages/DiscoveryResponse.h"
+#include "Containers/ArrayView.h"
+
+const uint32 FDiscoveryResponse::MinPayloadSize = 16 + 2 + 2; // Server Id (16) + Control Port (2) + Supported Versions (at least 1)
+FDiscoveryResponse::FDiscoveryResponse(FServerId InServerId, uint16 InControlPort, TArray<uint16> InSupportedVersions)
+    : ServerId(MoveTemp(InServerId))
+    , ControlPort(InControlPort)
+	, SupportedVersions(MoveTemp(InSupportedVersions))
+{
+}
+
+TProtocolResult<FDiscoveryResponse> FDiscoveryResponse::Deserialize(const FDiscoveryPacket& InPacket)
+{
+	if (InPacket.GetMessageType() != FDiscoveryPacket::EMessageType::Response)
+	{
+		return FCaptureProtocolError(TEXT("Invalid request arrived"));
+	}
+
+	TConstArrayView<uint8> Payload = InPacket.GetPayload();
+
+	if (Payload.Num() < MinPayloadSize)
+    {
+        return FCaptureProtocolError(TEXT("Invalid number of bytes"));
+    }
+
+    FServerId ServerId;
+    FMemory::Memcpy(ServerId.GetData(), Payload.GetData(), ServerId.Num());
+
+    uint16 ControlPort = *reinterpret_cast<const uint16*>(Payload.GetData() + ServerId.Num());
+
+	TConstArrayView<uint8> SupportedVersionsBytes = Payload.RightChop(ServerId.Num() + sizeof(ControlPort));
+	TArray<uint16> SupportedVersions;
+	SupportedVersions.Append((uint16*) SupportedVersionsBytes.GetData(), SupportedVersionsBytes.Num() / 2);
+	
+    return FDiscoveryResponse(MoveTemp(ServerId), ControlPort, MoveTemp(SupportedVersions));
+}
+
+TProtocolResult<FDiscoveryPacket> FDiscoveryResponse::Serialize(const FDiscoveryResponse& InResponse)
+{
+    TArray<uint8> Payload;
+    Payload.Reserve(InResponse.ServerId.Num() + sizeof(InResponse.ControlPort) + InResponse.SupportedVersions.Num() * sizeof(uint16));
+
+    Payload.Append(InResponse.ServerId.GetData(), InResponse.ServerId.Num());
+    Payload.Append((uint8*) &InResponse.ControlPort, sizeof(InResponse.ControlPort));
+	Payload.Append((uint8*) InResponse.SupportedVersions.GetData(), InResponse.SupportedVersions.Num() * sizeof(uint16));
+
+    return FDiscoveryPacket(FDiscoveryPacket::EMessageType::Response, MoveTemp(Payload));
+}
+
+const FDiscoveryResponse::FServerId& FDiscoveryResponse::GetServerId() const
+{
+    return ServerId;
+}
+
+uint16 FDiscoveryResponse::GetControlPort() const
+{
+    return ControlPort;
+}
+
+const TArray<uint16>& FDiscoveryResponse::GetSupportedVersions() const
+{
+	return SupportedVersions;
+}
\ No newline at end of file
diff --git a/Engine/Plugins/VirtualProduction/CaptureManager/Source/CaptureProtocolStack/Private/ExportClient/Communication/ExportCommunication.cpp b/Engine/Plugins/VirtualProduction/CaptureManager/Source/CaptureProtocolStack/Private/ExportClient/Communication/ExportCommunication.cpp
new file mode 100644
--- /dev/null
+++ b/Engine/Plugins/VirtualProduction/CaptureManager/Source/CaptureProtocolStack/Private/ExportClient/Communication/ExportCommunication.cpp
@@ -0,0 +1,88 @@
+// Copyright Epic Games, Inc. All Rights Reserved.
+
+#include "ExportClient/Communication/ExportCommunication.h"
+#include "ExportClient/Definitions.h"
+#include "Templates/UnrealTemplate.h"
+
+FExportCommunication::FExportCommunication()
+{
+}
+
+TProtocolResult<void> FExportCommunication::Init()
+{
+	return Client.Init();
+}
+
+TProtocolResult<void> FExportCommunication::Start(const FString& InServerIp, const uint16 InServerPort)
+{
+	return Client.Start(InServerIp + TEXT(":") + FString::FromInt(InServerPort));
+}
+
+TProtocolResult<void> FExportCommunication::Stop()
+{
+	return Client.Stop();
+}
+
+bool FExportCommunication::IsRunning() const
+{
+	return Client.IsRunning();
+}
+
+TProtocolResult<void> FExportCommunication::SendRequest(FExportHeader InHeader, FExportRequest InRequest)
+{
+	FExportRequestPacket Packet = { MoveTemp(InHeader) , MoveTemp(InRequest) };
+	FTcpClientWriter Writer(Client);
+
+	TProtocolResult<void> SerializeResult = FExportRequestPacket::Serialize(Packet, Writer);
+	if (SerializeResult.HasError())
+	{
+		return SerializeResult.StealError();
+	}
+
+	return ResultOk;
+}
+
+TProtocolResult<FExportResponseHeader> FExportCommunication::ReceiveResponseHeader()
+{
+	FTcpClientReader Reader(Client);
+
+	return FExportResponseHeader::Deserialize(Reader);
+}
+
+TProtocolResult<TArray<uint8>> FExportCommunication::ReceiveResponseData(const uint64 InSize)
+{
+	return Client.ReceiveMessage(InSize, InactivityTimeoutMs);
+}
+
+TProtocolResult<TStaticArray<uint8, 16>> FExportCommunication::ReceiveFileHash()
+{
+	FTcpClientReader Reader(Client);
+
+	return FExportResponse::DeserializeHash(Reader);
+}
+
+TProtocolResult<void> FExportRequestPacket::Serialize(const FExportRequestPacket& InRequestPacket, ITcpSocketWriter& InWriter)
+{
+	CPS_CHECK_VOID_RESULT(FExportHeader::Serialize(InRequestPacket.Header, InWriter));
+	CPS_CHECK_VOID_RESULT(FExportRequest::Serialize(InRequestPacket.Request, InWriter));
+
+	return ResultOk;
+}
+
+TProtocolResult<FExportResponseHeader> FExportResponseHeader::Deserialize(ITcpSocketReader& InReader)
+{
+	TProtocolResult<FExportHeader> HeaderResult = FExportHeader::Deserialize(InReader);
+	if (HeaderResult.HasError())
+	{
+		return HeaderResult.StealError();
+	}
+
+	TProtocolResult<FExportResponse> ResponseResult = FExportResponse::Deserialize(InReader);
+	if (ResponseResult.HasError())
+	{
+		return ResponseResult.StealError();
+	}
+
+	FExportResponseHeader ResponseHeader = { HeaderResult.StealValue(), ResponseResult.StealValue() };
+	return ResponseHeader;
+}
\ No newline at end of file
diff --git a/Engine/Plugins/VirtualProduction/CaptureManager/Source/CaptureProtocolStack/Private/ExportClient/Communication/ExportHeader.cpp b/Engine/Plugins/VirtualProduction/CaptureManager/Source/CaptureProtocolStack/Private/ExportClient/Communication/ExportHeader.cpp
new file mode 100644
--- /dev/null
+++ b/Engine/Plugins/VirtualProduction/CaptureManager/Source/CaptureProtocolStack/Private/ExportClient/Communication/ExportHeader.cpp
@@ -0,0 +1,68 @@
+// Copyright Epic Games, Inc. All Rights Reserved.
+
+#include "ExportClient/Communication/ExportHeader.h"
+#include "ExportClient/Definitions.h"
+
+const TArray<uint8> FExportHeader::Header = { 'C', 'P', 'S', 'E', 'X', 'P', 'O', 'R', 'T', '\0' };
+
+#define CPS_CHECK_RESULT_RET_ERROR(Result) if (Result.HasError()) { return Result.StealError(); }
+
+FExportHeader::FExportHeader(uint16 InVersion, uint32 InTransactionId)
+	: Version(InVersion)
+	, TransactionId(InTransactionId)
+{
+}
+
+TProtocolResult<FExportHeader> FExportHeader::Deserialize(ITcpSocketReader& InReader)
+{
+	FExportHeader ExportHeader;
+
+	TProtocolResult<TArray<uint8>> HeaderResult = InReader.ReceiveMessage(Header.Num(), InactivityTimeoutMs);
+	CPS_CHECK_RESULT_RET_ERROR(HeaderResult);
+
+	TArray<uint8> HeaderData = HeaderResult.StealValue();
+
+	if (FMemory::Memcmp(HeaderData.GetData(), Header.GetData(), Header.Num()) != 0)
+	{
+		return FCaptureProtocolError(TEXT("Header doesn't match"));
+	}
+
+	TProtocolResult<TArray<uint8>> VersionResult = InReader.ReceiveMessage(sizeof(ExportHeader.Version), InactivityTimeoutMs);
+	CPS_CHECK_RESULT_RET_ERROR(VersionResult);
+
+	TArray<uint8> VersionData = VersionResult.StealValue();
+
+	ExportHeader.Version = *reinterpret_cast<const uint16*>(VersionData.GetData());
+
+	TProtocolResult<TArray<uint8>> TransactionIdResult = InReader.ReceiveMessage(sizeof(ExportHeader.TransactionId), InactivityTimeoutMs);
+	CPS_CHECK_RESULT_RET_ERROR(TransactionIdResult);
+
+	TArray<uint8> TransactionIdData = TransactionIdResult.StealValue();
+	
+	ExportHeader.TransactionId = *reinterpret_cast<const uint32*>(TransactionIdData.GetData());
+
+	return ExportHeader;
+}
+
+TProtocolResult<void> FExportHeader::Serialize(const FExportHeader& InHeader, ITcpSocketWriter& InWriter)
+{
+	uint16 Version = InHeader.GetVersion();
+	uint32 TransactionId = InHeader.GetTransactionId();
+
+	TArray<uint8> Data;
+	Data.Append(Header.GetData(), Header.Num());
+	Data.Append(reinterpret_cast<uint8*>(&Version), sizeof(Version));
+	Data.Append(reinterpret_cast<uint8*>(&TransactionId), sizeof(TransactionId));
+
+	return InWriter.SendMessage(Data);
+}
+
+uint16 FExportHeader::GetVersion() const
+{
+	return Version;
+}
+
+uint32 FExportHeader::GetTransactionId() const
+{
+	return TransactionId;
+}
\ No newline at end of file
diff --git a/Engine/Plugins/VirtualProduction/CaptureManager/Source/CaptureProtocolStack/Private/ExportClient/Definitions.h b/Engine/Plugins/VirtualProduction/CaptureManager/Source/CaptureProtocolStack/Private/ExportClient/Definitions.h
new file mode 100644
--- /dev/null
+++ b/Engine/Plugins/VirtualProduction/CaptureManager/Source/CaptureProtocolStack/Private/ExportClient/Definitions.h
@@ -0,0 +1,5 @@
+// Copyright Epic Games, Inc. All Rights Reserved.
+
+#pragma once
+
+static constexpr uint32 InactivityTimeoutMs = 15000; // 15 seconds as recommended by the CPS Specification
\ No newline at end of file
diff --git a/Engine/Plugins/VirtualProduction/CaptureManager/Source/CaptureProtocolStack/Private/ExportClient/ExportClient.cpp b/Engine/Plugins/VirtualProduction/CaptureManager/Source/CaptureProtocolStack/Private/ExportClient/ExportClient.cpp
new file mode 100644
--- /dev/null
+++ b/Engine/Plugins/VirtualProduction/CaptureManager/Source/CaptureProtocolStack/Private/ExportClient/ExportClient.cpp
@@ -0,0 +1,298 @@
+// Copyright Epic Games, Inc. All Rights Reserved.
+
+#include "ExportClient/ExportClient.h"
+
+#include "ExportClient/Messages/ExportRequest.h"
+
+#include "Utility/Definitions.h"
+
+
+FExportTaskStopToken::FExportTaskStopToken()
+	: bCanceled(false)
+{
+}
+
+void FExportTaskStopToken::Cancel()
+{
+	bCanceled.store(true);
+}
+
+bool FExportTaskStopToken::IsCanceled()
+{
+	return bCanceled.load();
+}
+
+FExportClient::FExportClient(FString InServerIp, const uint16 InExportPort)
+	: ServerIp(MoveTemp(InServerIp))
+	, ServerPort(InExportPort)
+	, Communication(MakeUnique<FExportCommunication>())
+	, ExportWorker(*this)
+	, TransactionIdCounter(0)
+{
+	Thread.Reset(FRunnableThread::Create(&ExportWorker, TEXT("Queue Runner"), 128 * 1024, TPri_Normal, FPlatformAffinity::GetPoolThreadMask()));
+}
+
+FExportClient::~FExportClient()
+{
+	Thread->Kill(true);
+	Stop();
+}
+
+FExportClient::FTaskId FExportClient::ExportTakeFiles(FString InTakeName, FTakeFileArray InTakeFileArray, TUniquePtr<FBaseStream> InStream)
+{
+	check(InStream);
+	check(ServerPort != 0);
+
+	FExportTakeTask::FExportContexts ExportContexts;
+	for (FTakeFile& File : InTakeFileArray)
+	{
+		ExportContexts.Add({ InTakeName, MoveTemp(File) });
+	}
+	
+	return StartExport(MoveTemp(ExportContexts), MoveTemp(InStream));
+}
+
+FExportClient::FTaskId FExportClient::ExportFiles(TMap<FString, FTakeFileArray> InTakesFilesMap, TUniquePtr<FBaseStream> InStream)
+{
+	check(InStream);
+	check(ServerPort != 0);
+
+	FExportTakeTask::FExportContexts ExportContexts;
+	for (TPair<FString, FTakeFileArray>& TakeFiles : InTakesFilesMap)
+	{
+		for (FTakeFile& File : TakeFiles.Value)
+		{
+			ExportContexts.Add({ TakeFiles.Key, MoveTemp(File) });
+		}
+	}
+	
+	return StartExport(MoveTemp(ExportContexts), MoveTemp(InStream));
+}
+
+void FExportClient::AbortExport(const FTaskId InTaskId)
+{
+	TProtocolResult<TUniquePtr<FExportTakeTask>> RemovedTaskResult = ExportWorker.Remove(InTaskId);
+	if (RemovedTaskResult.HasValue())
+	{
+		TUniquePtr<FExportTakeTask> RemovedTask = RemovedTaskResult.StealValue();
+
+		RemovedTask->Stream->Done(FCaptureProtocolError(TEXT("Take aborted")));
+	}
+	else
+	{
+		CurrentTaskStopToken.IsValid() ? CurrentTaskStopToken->Cancel() : void();
+	}
+}
+
+void FExportClient::AbortAllExports()
+{
+	TArray<TUniquePtr<FExportTakeTask>> Tasks = ExportWorker.GetAndEmpty();
+
+	for (const TUniquePtr<FExportTakeTask>& Task : Tasks)
+	{
+		Task->Stream->Done(FCaptureProtocolError(TEXT("Take aborted")));
+	}
+
+	CurrentTaskStopToken.IsValid() ? CurrentTaskStopToken->Cancel() : void();
+}
+
+FExportClient::FTaskId FExportClient::StartExport(FExportTakeTask::FExportContexts InContexts, TUniquePtr<FBaseStream> InStream)
+{
+	TUniquePtr<FExportTakeTask> ExportTakeTask =
+		MakeUnique<FExportTakeTask>(MoveTemp(InContexts), MoveTemp(InStream));
+
+	FTaskId TaskIdToReturn = CurrentTaskId.fetch_add(1);
+
+	ExportWorker.Add(TaskIdToReturn, MoveTemp(ExportTakeTask));
+
+	return TaskIdToReturn;
+}
+
+TProtocolResult<void> FExportClient::Start()
+{
+	if (!Communication->IsRunning())
+	{
+		CPS_CHECK_VOID_RESULT(Communication->Init());
+
+		CPS_CHECK_VOID_RESULT(Communication->Start(ServerIp, ServerPort));
+	}
+
+	return ResultOk;
+}
+
+TProtocolResult<void> FExportClient::Stop()
+{
+	if (Communication->IsRunning())
+	{
+		CPS_CHECK_VOID_RESULT(Communication->Stop());
+	}
+
+	return ResultOk;
+}
+
+void FExportClient::OnTask(TUniquePtr<FExportTakeTask> InTask)
+{
+	TProtocolResult<void> StartResult = Start();
+	if (StartResult.HasError())
+	{
+		InTask->Stream->Done(StartResult.StealError());
+		return;
+	}
+
+	CurrentTaskStopToken = MakeUnique<FExportTaskStopToken>();
+
+	bool bShouldKeepConnection = OnExportTask(MoveTemp(InTask));
+
+	if (ExportWorker.IsEmpty() || 
+		CurrentTaskStopToken->IsCanceled() || 
+		!bShouldKeepConnection)
+	{
+		CurrentTaskStopToken = nullptr;
+
+		Stop();
+	}
+}
+
+bool FExportClient::OnExportTask(TUniquePtr<FExportTakeTask> InTask)
+{
+	if (!Communication.IsValid())
+	{
+		InTask->Stream->Done(FCaptureProtocolError(TEXT("Failed to start export")));
+		return false;
+	}
+
+	// Send
+	TProtocolResult<TMap<uint32, TPair<FString, FTakeFile>>> SendResult = SendRequests(MoveTemp(InTask->ExportContexts));
+	if (SendResult.HasError())
+	{
+		InTask->Stream->Done(SendResult.StealError());
+		return false;
+	}
+
+	TMap<uint32, TPair<FString, FTakeFile>> RequestToResponseData = SendResult.StealValue();
+
+	// Receive
+	TProtocolResult<void> ReceiveResult = ReceiveResponses(MoveTemp(RequestToResponseData), *InTask->Stream);
+	if (ReceiveResult.HasError())
+	{
+		InTask->Stream->Done(ReceiveResult.StealError());
+		return false;
+	}
+
+	InTask->Stream->Done(ResultOk);
+
+	return true;
+}
+
+TProtocolResult<TMap<uint32, TPair<FString, FTakeFile>>> FExportClient::SendRequests(FExportTakeTask::FExportContexts InContexts)
+{
+	TMap<uint32, TPair<FString, FTakeFile>> ResponseMap;
+
+	for (FExportTakeTask::FExportContext& ExportContext : InContexts)
+	{
+		if (CurrentTaskStopToken->IsCanceled())
+		{
+			return FCaptureProtocolError(TEXT("Take aborted"));
+		}
+
+		uint32 TransactionId = TransactionIdCounter++;
+
+		FExportHeader ExportHeader(CPS_VERSION, TransactionId);
+		FExportRequest ExportRequest(ExportContext.TakeName, ExportContext.File.FileName, ExportContext.File.Offset);
+
+		TProtocolResult<void> SendRequestResult = Communication->SendRequest(MoveTemp(ExportHeader), MoveTemp(ExportRequest));
+
+		if (SendRequestResult.HasError())
+		{
+			return SendRequestResult.StealError();
+		}
+
+		ResponseMap.Emplace(TransactionId, { ExportContext.TakeName, MoveTemp(ExportContext.File) });
+	}
+
+	return ResponseMap;
+}
+
+TProtocolResult<void> FExportClient::ReceiveResponses(TMap<uint32, TPair<FString, FTakeFile>> InResponseMap,
+													  FBaseStream& InBaseStream)
+{
+	while (!InResponseMap.IsEmpty())
+	{
+		if (CurrentTaskStopToken->IsCanceled())
+		{
+			return FCaptureProtocolError(TEXT("Take aborted"));
+		}
+
+		TProtocolResult<FExportResponseHeader> ResponseHeaderResult = Communication->ReceiveResponseHeader();
+
+		if (ResponseHeaderResult.HasError())
+		{
+			return ResponseHeaderResult.StealError();
+		}
+
+		FExportResponseHeader ResponseHeader = ResponseHeaderResult.StealValue();
+
+		uint32 ResponseTransactionId = ResponseHeader.Header.GetTransactionId();
+		FExportResponse::EStatus Status = ResponseHeader.Response.GetStatus();
+
+		if (Status != FExportResponse::EStatus::Success)
+		{
+			return FCaptureProtocolError(TEXT("Server responded with error status"), static_cast<int32>(Status));
+		}
+
+		const TPair<FString, FTakeFile>& TakeFile = InResponseMap[ResponseTransactionId];
+
+		const FTakeFile& File = TakeFile.Value;
+		if (!InBaseStream.StartFile(TakeFile.Key, File.FileName))
+		{
+			// Error is already handled
+			return FCaptureProtocolError(TEXT("Stream error"));
+		}
+
+		uint64 BytesLeft = File.Length;
+
+		constexpr uint32 MaxChunkSize = 64 * 1024;
+
+		while (BytesLeft != 0)
+		{
+			if (CurrentTaskStopToken->IsCanceled())
+			{
+				return FCaptureProtocolError(TEXT("Take aborted"));
+			}
+
+			uint32 ChunkSize = BytesLeft > MaxChunkSize ? MaxChunkSize : BytesLeft;
+
+			TProtocolResult<TArray<uint8>> FileDataResult = Communication->ReceiveResponseData(ChunkSize);
+
+			if (FileDataResult.HasError())
+			{
+				return FileDataResult.StealError();
+			}
+
+			TArray<uint8> FileData = FileDataResult.StealValue();
+
+			if (!InBaseStream.ProcessData(TakeFile.Key, File.FileName, FileData))
+			{
+				// Error is already handled
+				return FCaptureProtocolError(TEXT("Stream error"));
+			}
+
+			BytesLeft -= FileData.Num();
+		}
+
+		TProtocolResult<TStaticArray<uint8, 16>> FileHashResult = Communication->ReceiveFileHash();
+		if (FileHashResult.HasError())
+		{
+			return FileHashResult.StealError();
+		}
+
+		if (!InBaseStream.FinishFile(TakeFile.Key, File.FileName, FileHashResult.StealValue()))
+		{
+			return FCaptureProtocolError(TEXT("Stream error"));
+		}
+
+		InResponseMap.Remove(ResponseTransactionId);
+	}
+
+	return ResultOk;
+}
\ No newline at end of file
diff --git a/Engine/Plugins/VirtualProduction/CaptureManager/Source/CaptureProtocolStack/Private/ExportClient/ExportWorker.cpp b/Engine/Plugins/VirtualProduction/CaptureManager/Source/CaptureProtocolStack/Private/ExportClient/ExportWorker.cpp
new file mode 100644
--- /dev/null
+++ b/Engine/Plugins/VirtualProduction/CaptureManager/Source/CaptureProtocolStack/Private/ExportClient/ExportWorker.cpp
@@ -0,0 +1,146 @@
+// Copyright Epic Games, Inc. All Rights Reserved.
+
+#include "ExportClient/ExportWorker.h"
+
+FExportTakeTask::FExportTakeTask(FExportContexts InExportContexts,
+								 TUniquePtr<FBaseStream> InStream)
+	: ExportContexts(MoveTemp(InExportContexts))
+	, Stream(MoveTemp(InStream))
+{
+}
+
+FExportQueue::FExportQueue()
+	: Semaphore(0, MaxNumberOfElements) {
+}
+
+
+TProtocolResult<void> FExportQueue::Add(const uint32 InTaskId, TUniquePtr<FExportTakeTask> InElement)
+{
+	{
+		FScopeLock Lock(&Mutex);
+		if (Queue.Num() >= MaxNumberOfElements)
+		{
+			return FCaptureProtocolError(TEXT("Number of elements exceeded"));
+		}
+
+		Queue.Add({ InTaskId, MoveTemp(InElement) });
+	}
+
+	Semaphore.Release();
+
+	return ResultOk;
+}
+
+TUniquePtr<FExportTakeTask> FExportQueue::Pop()
+{
+	while (true)
+	{
+		Semaphore.Acquire();
+
+		FScopeLock Lock(&Mutex);
+		if (!Queue.IsEmpty())
+		{
+			TUniquePtr<FExportTakeTask> RemovedItem = MoveTemp(Queue[0].Value);
+			Queue.RemoveAt(0);
+
+			return MoveTemp(RemovedItem);
+		}
+	}
+}
+
+TProtocolResult<TUniquePtr<FExportTakeTask>> FExportQueue::Remove(uint32 InTaskId)
+{
+	int32 Index = 0;
+
+	FScopeLock Lock(&Mutex);
+	for (; Index < Queue.Num(); ++Index)
+	{
+		if (Queue[Index].Key == InTaskId)
+		{
+			break;
+		}
+	}
+
+	if (Index == Queue.Num())
+	{
+		return FCaptureProtocolError(TEXT("Element doesn't exist"));
+	}
+
+	TUniquePtr<FExportTakeTask> RemovedItem = MoveTemp(Queue[Index].Value);
+	Queue.RemoveAt(Index);
+
+	return RemovedItem;
+}
+
+bool FExportQueue::IsEmpty() const
+{
+	FScopeLock Lock(&Mutex);
+	return Queue.IsEmpty();
+}
+
+TArray<TUniquePtr<FExportTakeTask>> FExportQueue::GetAndEmpty()
+{
+	TArray<TUniquePtr<FExportTakeTask>> TaskArray;
+
+	FScopeLock Lock(&Mutex);
+
+	for (TPair<uint32, TUniquePtr<FExportTakeTask>>& Task : Queue)
+	{
+		TaskArray.Add(MoveTemp(Task.Value));
+	}
+
+	Queue.Empty();
+
+	return TaskArray;
+}
+
+FExportWorker::FExportWorker(FExportTaskExecutor& InClient)
+	: bRunning(true)
+	, Client(InClient)
+{
+}
+
+FExportWorker::~FExportWorker()
+{
+}
+
+TProtocolResult<void> FExportWorker::Add(const uint32 InTaskId, TUniquePtr<FExportTakeTask> InElement)
+{
+	return Queue.Add(InTaskId, MoveTemp(InElement));
+}
+
+TProtocolResult<TUniquePtr<FExportTakeTask>> FExportWorker::Remove(uint32 InTaskId)
+{
+	return Queue.Remove(InTaskId);
+}
+
+bool FExportWorker::IsEmpty() const
+{
+	return Queue.IsEmpty();
+}
+
+TArray<TUniquePtr<FExportTakeTask>> FExportWorker::GetAndEmpty()
+{
+	return Queue.GetAndEmpty();
+}
+
+uint32 FExportWorker::Run()
+{
+	while (bRunning.load())
+	{
+		TUniquePtr<FExportTakeTask> Elem = Queue.Pop();
+		if (Elem.IsValid())
+		{
+			Client.OnTask(MoveTemp(Elem));
+		}
+	}
+
+	return 0;
+}
+
+void FExportWorker::Stop()
+{
+	bRunning.store(false);
+
+	Queue.Add(0, nullptr);
+}
diff --git a/Engine/Plugins/VirtualProduction/CaptureManager/Source/CaptureProtocolStack/Private/ExportClient/Messages/ExportRequest.cpp b/Engine/Plugins/VirtualProduction/CaptureManager/Source/CaptureProtocolStack/Private/ExportClient/Messages/ExportRequest.cpp
new file mode 100644
--- /dev/null
+++ b/Engine/Plugins/VirtualProduction/CaptureManager/Source/CaptureProtocolStack/Private/ExportClient/Messages/ExportRequest.cpp
@@ -0,0 +1,90 @@
+// Copyright Epic Games, Inc. All Rights Reserved.
+
+#include "ExportClient/Messages/ExportRequest.h"
+
+#include "Containers/StringConv.h"
+
+#define CPS_CHECK_RESULT_RET_ERROR(Result) if (Result.HasError()) { return Result.StealError(); }
+
+FExportRequest::FExportRequest(FString InTakeName, FString InFileName, uint64 InOffset)
+	: TakeName(MoveTemp(InTakeName))
+	, FileName(MoveTemp(InFileName))
+	, Offset(InOffset)
+{
+}
+
+TProtocolResult<FExportRequest> FExportRequest::Deserialize(ITcpSocketReader& InReader)
+{
+	FExportRequest Request;
+
+	uint16 Length;
+	TProtocolResult<TArray<uint8>> TakeNameLengthResult = InReader.ReceiveMessage(sizeof(Length));
+	CPS_CHECK_RESULT_RET_ERROR(TakeNameLengthResult);
+
+	Length = *reinterpret_cast<uint16*>(TakeNameLengthResult.StealValue().GetData());
+
+	TProtocolResult<TArray<uint8>> TakeNameResult = InReader.ReceiveMessage(Length);
+	CPS_CHECK_RESULT_RET_ERROR(TakeNameResult);
+
+	TArray<uint8> TakeNameData = TakeNameResult.StealValue();
+	Request.TakeName = FString(Length, (const UTF8CHAR*) TakeNameData.GetData());
+
+	TProtocolResult<TArray<uint8>> FileNameLengthResult = InReader.ReceiveMessage(sizeof(Length));
+	CPS_CHECK_RESULT_RET_ERROR(FileNameLengthResult);
+
+	Length = *reinterpret_cast<uint16*>(FileNameLengthResult.StealValue().GetData());
+
+	TProtocolResult<TArray<uint8>> FileNameResult = InReader.ReceiveMessage(Length);
+	CPS_CHECK_RESULT_RET_ERROR(FileNameResult);
+
+	TArray<uint8> FileNameData = FileNameResult.StealValue();
+	Request.FileName = FString(Length, (const UTF8CHAR*) FileNameData.GetData());
+
+	TProtocolResult<TArray<uint8>> OffsetResult = InReader.ReceiveMessage(sizeof(Request.Offset));
+	CPS_CHECK_RESULT_RET_ERROR(OffsetResult);
+
+	TArray<uint8> OffsetData = OffsetResult.StealValue();
+	
+	Request.Offset = *reinterpret_cast<uint64*>(OffsetData.GetData());
+
+	return Request;
+}
+
+TProtocolResult<void> FExportRequest::Serialize(const FExportRequest& InRequest, ITcpSocketWriter& InWriter)
+{
+	const FString& TakeName = InRequest.GetTakeName();
+	const FString& FileName = InRequest.GetFileName();
+
+	const uint64 Offset = InRequest.GetOffset();
+
+	TArray<uint8> Data;
+
+	auto TakeNameUTF8 = StringCast<UTF8CHAR>(*TakeName);
+	auto FileNameUTF8 = StringCast<UTF8CHAR>(*FileName);
+
+	const uint16 TakeNameLen = TakeNameUTF8.Length();
+	const uint16 FileNameLen = FileNameUTF8.Length();
+
+	Data.Append(reinterpret_cast<const uint8*>(&TakeNameLen), sizeof(TakeNameLen));
+	Data.Append(reinterpret_cast<const uint8*>(TakeNameUTF8.Get()), TakeNameLen);
+	Data.Append(reinterpret_cast<const uint8*>(&FileNameLen), sizeof(FileNameLen));
+	Data.Append(reinterpret_cast<const uint8*>(FileNameUTF8.Get()), FileNameLen);
+	Data.Append(reinterpret_cast<const uint8*>(&Offset), sizeof(Offset));
+
+	return InWriter.SendMessage(Data);
+}
+
+const FString& FExportRequest::GetTakeName() const
+{
+	return TakeName;
+}
+
+const FString& FExportRequest::GetFileName() const
+{
+	return FileName;
+}
+
+uint64 FExportRequest::GetOffset() const
+{
+	return Offset;
+}
\ No newline at end of file
diff --git a/Engine/Plugins/VirtualProduction/CaptureManager/Source/CaptureProtocolStack/Private/ExportClient/Messages/ExportResponse.cpp b/Engine/Plugins/VirtualProduction/CaptureManager/Source/CaptureProtocolStack/Private/ExportClient/Messages/ExportResponse.cpp
new file mode 100644
--- /dev/null
+++ b/Engine/Plugins/VirtualProduction/CaptureManager/Source/CaptureProtocolStack/Private/ExportClient/Messages/ExportResponse.cpp
@@ -0,0 +1,131 @@
+// Copyright Epic Games, Inc. All Rights Reserved.
+
+#include "ExportClient/Messages/ExportResponse.h"
+#include "ExportClient/Definitions.h"
+
+#define CPS_CHECK_RESULT_RET_ERROR(Result) if (Result.HasError()) { return Result.StealError(); }
+
+FExportResponse::FExportResponse(EStatus InStatus, uint64 InLength)
+	: Status(InStatus)
+	, Length(InLength)
+{
+}
+
+TProtocolResult<FExportResponse> FExportResponse::Deserialize(ITcpSocketReader& InReader)
+{
+	FExportResponse Response;
+
+	TProtocolResult<TArray<uint8>> StatusResult = InReader.ReceiveMessage(sizeof(Response.Status), InactivityTimeoutMs);
+	CPS_CHECK_RESULT_RET_ERROR(StatusResult);
+
+	Response.Status = ToStatus(StatusResult.StealValue()[0]);
+
+	if (Response.Status == EStatus::Success)
+	{
+		TProtocolResult<TArray<uint8>> LengthResult = InReader.ReceiveMessage(sizeof(Response.Length), InactivityTimeoutMs);
+		CPS_CHECK_RESULT_RET_ERROR(LengthResult);
+
+		TArray<uint8> LengthData = LengthResult.StealValue();
+
+		Response.Length = *reinterpret_cast<const uint64*>(LengthData.GetData());
+	}
+
+	return Response;
+}
+
+TProtocolResult<TStaticArray<uint8, 16>> FExportResponse::DeserializeHash(ITcpSocketReader& InReader)
+{
+	TStaticArray<uint8, 16> Hash;
+
+	TProtocolResult<TArray<uint8>> HashResult = InReader.ReceiveMessage(Hash.Num(), InactivityTimeoutMs);
+	CPS_CHECK_RESULT_RET_ERROR(HashResult);
+
+	TArray<uint8> HashData = HashResult.StealValue();
+
+	FMemory::Memcpy(Hash.GetData(), HashData.GetData(), Hash.Num());
+
+	return Hash;
+}
+
+TProtocolResult<void> FExportResponse::Serialize(const FExportResponse& InResponse, ITcpSocketWriter& InWriter)
+{
+	uint8 Status = FromStatus(InResponse.GetStatus());
+
+	TArray<uint8> Data;
+	Data.Append(reinterpret_cast<uint8*>(&Status), sizeof(Status));
+
+	if (InResponse.GetStatus() == EStatus::Success)
+	{
+		uint64 Length = InResponse.GetLength();
+		Data.Append(reinterpret_cast<uint8*>(&Length), sizeof(Length));
+	}
+
+	return InWriter.SendMessage(Data);
+}
+
+TProtocolResult<void> FExportResponse::SerializeHash(const TStaticArray<uint8, 16>& InHash, ITcpSocketWriter& InWriter)
+{
+	TArray<uint8> Data;
+
+	Data.Append(InHash.GetData(), InHash.Num());
+
+	return InWriter.SendMessage(Data);
+}
+
+FExportResponse::EStatus FExportResponse::GetStatus() const
+{
+	return Status;
+}
+
+uint64 FExportResponse::GetLength() const
+{
+	return Length;
+}
+
+FExportResponse::EStatus FExportResponse::ToStatus(uint8 InStatus)
+{
+	if (InStatus == 0)
+	{
+		return EStatus::Success;
+	}
+	else if (InStatus == 1)
+	{
+		return EStatus::InvalidTakeName;
+	}
+	else if (InStatus == 2)
+	{
+		return EStatus::InvalidFileName;
+	}
+	else if (InStatus == 3)
+	{
+		return EStatus::InvalidOffset;
+	}
+	else if (InStatus == 4)
+	{
+		return EStatus::ServerError;
+	}
+	else
+	{
+		return EStatus::Reserved;
+	}
+}
+
+uint8 FExportResponse::FromStatus(EStatus InStatus)
+{
+	switch (InStatus)
+	{
+		case EStatus::Success:
+			return 0;
+		case EStatus::InvalidTakeName:
+			return 1;
+		case EStatus::InvalidFileName:
+			return 2;
+		case EStatus::InvalidOffset:
+			return 3;
+		case EStatus::ServerError:
+			return 4;
+		case EStatus::Reserved:
+		default:
+			return 255;
+	}
+}
\ No newline at end of file
diff --git a/Engine/Plugins/VirtualProduction/CaptureManager/Source/CaptureProtocolStack/Private/Tests/Control/TestControlJsonUtilities.cpp b/Engine/Plugins/VirtualProduction/CaptureManager/Source/CaptureProtocolStack/Private/Tests/Control/TestControlJsonUtilities.cpp
new file mode 100644
--- /dev/null
+++ b/Engine/Plugins/VirtualProduction/CaptureManager/Source/CaptureProtocolStack/Private/Tests/Control/TestControlJsonUtilities.cpp
@@ -0,0 +1,257 @@
+// Copyright Epic Games, Inc. All Rights Reserved.
+
+#include "Control/Messages/ControlJsonUtilities.h"
+
+#include "Misc/AutomationTest.h"
+
+#if WITH_DEV_AUTOMATION_TESTS
+IMPLEMENT_SIMPLE_AUTOMATION_TEST(TControlJsonUtilitiesParseNumber, "Plugin.CaptureProtocolStack.Control.ControlJsonUtilities.ParseNumber.Success", EAutomationTestFlags::ApplicationContextMask | EAutomationTestFlags::SmokeFilter)
+bool TControlJsonUtilitiesParseNumber::RunTest(const FString& InParameters)
+{
+    FString KeyInteger = TEXT("KeyInteger");
+    int32 ValueInteger = 10;
+    FString KeyFloat = TEXT("KeyFloat");
+    float ValueFloat = 10.0;
+
+    TSharedPtr<FJsonObject> Object = MakeShared<FJsonObject>();
+    Object->SetNumberField(KeyInteger, ValueInteger);
+    Object->SetNumberField(KeyFloat, ValueFloat);
+
+    int32 ParseValueInteger;
+    TestTrue(TEXT("ParseNumber"), FJsonUtility::ParseNumber(Object, KeyInteger, ParseValueInteger).HasValue());
+    TestEqual(TEXT("ParseNumber"), ParseValueInteger, ValueInteger);
+
+    float ParseValueFloat;
+    TestTrue(TEXT("ParseNumber"), FJsonUtility::ParseNumber(Object, KeyFloat, ParseValueFloat).HasValue());
+    TestEqual(TEXT("ParseNumber"), ParseValueFloat, ValueFloat);
+
+    return true;
+}
+
+IMPLEMENT_SIMPLE_AUTOMATION_TEST(TControlJsonUtilitiesParseNumberMissingField, "Plugin.CaptureProtocolStack.Control.ControlJsonUtilities.ParseNumber.MissingField", EAutomationTestFlags::ApplicationContextMask | EAutomationTestFlags::SmokeFilter)
+bool TControlJsonUtilitiesParseNumberMissingField::RunTest(const FString& InParameters)
+{
+    FString Key = TEXT("KeyInteger");
+    FString KeyInvalid = TEXT("KeyIntegerInvalid");
+    int32 Value = 10;
+
+    TSharedPtr<FJsonObject> Object = MakeShared<FJsonObject>();
+    Object->SetNumberField(Key, Value);
+
+    int32 ParseValue;
+    TestTrue(TEXT("ParseNumber"), FJsonUtility::ParseNumber(Object, KeyInvalid, ParseValue).HasError());
+
+    return true;
+}
+
+IMPLEMENT_SIMPLE_AUTOMATION_TEST(TControlJsonUtilitiesParseString, "Plugin.CaptureProtocolStack.Control.ControlJsonUtilities.ParseString.Success", EAutomationTestFlags::ApplicationContextMask | EAutomationTestFlags::SmokeFilter)
+bool TControlJsonUtilitiesParseString::RunTest(const FString& InParameters)
+{
+    FString Key = TEXT("KeyString");
+    FString Value = TEXT("ValueString");
+
+    TSharedPtr<FJsonObject> Object = MakeShared<FJsonObject>();
+    Object->SetStringField(Key, Value);
+
+    FString ParseValue;
+    TestTrue(TEXT("ParseString"), FJsonUtility::ParseString(Object, Key, ParseValue).HasValue());
+    TestEqual(TEXT("ParseString"), ParseValue, Value);
+
+    return true;
+}
+
+IMPLEMENT_SIMPLE_AUTOMATION_TEST(TControlJsonUtilitiesParseStringMissingField, "Plugin.CaptureProtocolStack.Control.ControlJsonUtilities.ParseString.MissingField", EAutomationTestFlags::ApplicationContextMask | EAutomationTestFlags::SmokeFilter)
+bool TControlJsonUtilitiesParseStringMissingField::RunTest(const FString& InParameters)
+{
+    FString Key = TEXT("KeyString");
+    FString KeyInvalid = TEXT("KeyStringInvalid");
+    FString Value = TEXT("ValueString");
+
+    TSharedPtr<FJsonObject> Object = MakeShared<FJsonObject>();
+    Object->SetStringField(Key, Value);
+
+    FString ParseValue;
+    TestTrue(TEXT("ParseString"), FJsonUtility::ParseString(Object, KeyInvalid, ParseValue).HasError());
+
+    return true;
+}
+
+IMPLEMENT_SIMPLE_AUTOMATION_TEST(TControlJsonUtilitiesParseBool, "Plugin.CaptureProtocolStack.Control.ControlJsonUtilities.ParseBool.Success", EAutomationTestFlags::ApplicationContextMask | EAutomationTestFlags::SmokeFilter)
+bool TControlJsonUtilitiesParseBool::RunTest(const FString& InParameters)
+{
+    FString Key = TEXT("KeyBool");
+    bool Value = true;
+
+    TSharedPtr<FJsonObject> Object = MakeShared<FJsonObject>();
+    Object->SetBoolField(Key, Value);
+
+    bool ParseValue;
+    TestTrue(TEXT("ParseBool"), FJsonUtility::ParseBool(Object, Key, ParseValue).HasValue());
+    TestEqual(TEXT("ParseBool"), ParseValue, Value);
+
+    return true;
+}
+
+IMPLEMENT_SIMPLE_AUTOMATION_TEST(TControlJsonUtilitiesParseBoolMissingField, "Plugin.CaptureProtocolStack.Control.ControlJsonUtilities.ParseString.MissingField", EAutomationTestFlags::ApplicationContextMask | EAutomationTestFlags::SmokeFilter)
+bool TControlJsonUtilitiesParseBoolMissingField::RunTest(const FString& InParameters)
+{
+    FString Key = TEXT("KeyBool");
+    FString KeyInvalid = TEXT("KeyInvalid");
+    bool Value = true;
+
+    TSharedPtr<FJsonObject> Object = MakeShared<FJsonObject>();
+    Object->SetBoolField(Key, Value);
+
+    bool ParseValue;
+    TestTrue(TEXT("ParseBool"), FJsonUtility::ParseBool(Object, KeyInvalid, ParseValue).HasError());
+
+    return true;
+}
+
+IMPLEMENT_SIMPLE_AUTOMATION_TEST(TControlJsonUtilitiesParseObject, "Plugin.CaptureProtocolStack.Control.ControlJsonUtilities.ParseObject.Success", EAutomationTestFlags::ApplicationContextMask | EAutomationTestFlags::SmokeFilter)
+bool TControlJsonUtilitiesParseObject::RunTest(const FString& InParameters)
+{
+    FString Key = TEXT("KeyObject");
+    TSharedPtr<FJsonObject> Value = MakeShared<FJsonObject>();
+    Value->SetStringField(TEXT("Field"), TEXT("Value"));
+
+    TSharedPtr<FJsonObject> Object = MakeShared<FJsonObject>();
+    Object->SetObjectField(Key, Value);
+
+    const TSharedPtr<FJsonObject>* ParseValue;
+    TestTrue(TEXT("ParseObject"), FJsonUtility::ParseObject(Object, Key, ParseValue).HasValue());
+
+    for (const TPair<FString, TSharedPtr<FJsonValue>>& Pair : Value->Values)
+    {
+        TestTrue(TEXT("ParseObject"), (*ParseValue)->HasField(Pair.Key));
+
+        TSharedPtr<FJsonValue> ParseFieldValue = (*ParseValue)->TryGetField(Pair.Key);
+
+        TestNotNull(TEXT("ParseObject"), ParseFieldValue.Get());
+        TestEqual(TEXT("ParseObject"), *Pair.Value, *ParseFieldValue);
+    }
+
+    return true;
+}
+
+IMPLEMENT_SIMPLE_AUTOMATION_TEST(TControlJsonUtilitiesParseObjectMissingField, "Plugin.CaptureProtocolStack.Control.ControlJsonUtilities.ParseObject.MissingField", EAutomationTestFlags::ApplicationContextMask | EAutomationTestFlags::SmokeFilter)
+bool TControlJsonUtilitiesParseObjectMissingField::RunTest(const FString& InParameters)
+{
+    FString Key = TEXT("KeyObject");
+    FString KeyInvalid = TEXT("KeyInvalid");
+    TSharedPtr<FJsonObject> Value = MakeShared<FJsonObject>();
+    Value->SetStringField(TEXT("Field"), TEXT("Value"));
+
+    TSharedPtr<FJsonObject> Object = MakeShared<FJsonObject>();
+    Object->SetObjectField(Key, Value);
+
+    const TSharedPtr<FJsonObject>* ParseValue;
+    TestTrue(TEXT("ParseObject"), FJsonUtility::ParseObject(Object, KeyInvalid, ParseValue).HasError());
+
+    return true;
+}
+
+IMPLEMENT_SIMPLE_AUTOMATION_TEST(TControlJsonUtilitiesParseArray, "Plugin.CaptureProtocolStack.Control.ControlJsonUtilities.ParseArray.Success", EAutomationTestFlags::ApplicationContextMask | EAutomationTestFlags::SmokeFilter)
+bool TControlJsonUtilitiesParseArray::RunTest(const FString& InParameters)
+{
+    FString Key = TEXT("KeyArray");
+
+    constexpr uint32 MaxNumberOfElemens = 10;
+    TArray<TSharedPtr<FJsonValue>> Value;
+    Value.Reserve(MaxNumberOfElemens);
+
+    for (int32 ArrayValue = 0; ArrayValue < MaxNumberOfElemens; ++ArrayValue)
+    {
+        Value.Add(MakeShared<FJsonValueNumber>(ArrayValue));
+    }
+
+    TSharedPtr<FJsonObject> Object = MakeShared<FJsonObject>();
+    Object->SetArrayField(Key, Value);
+
+    const TArray<TSharedPtr<FJsonValue>>* ParseValue;
+    TestTrue(TEXT("ParseArray"), FJsonUtility::ParseArray(Object, Key, ParseValue).HasValue());
+
+    for (int32 Index = 0; Index < MaxNumberOfElemens; ++Index)
+    {
+        TestEqual(TEXT("ParseArray"), Value[Index], (*ParseValue)[Index]);
+    }
+
+    return true;
+}
+
+IMPLEMENT_SIMPLE_AUTOMATION_TEST(TControlJsonUtilitiesParseArrayMissingField, "Plugin.CaptureProtocolStack.Control.ControlJsonUtilities.ParseArray.MissingField", EAutomationTestFlags::ApplicationContextMask | EAutomationTestFlags::SmokeFilter)
+bool TControlJsonUtilitiesParseArrayMissingField::RunTest(const FString& InParameters)
+{
+    FString Key = TEXT("KeyArray");
+    FString KeyInvalid = TEXT("KeyInvalid");
+
+    constexpr uint32 MaxNumberOfElemens = 10;
+    TArray<TSharedPtr<FJsonValue>> Value;
+    Value.Reserve(MaxNumberOfElemens);
+
+    for (int32 ArrayValue = 0; ArrayValue < MaxNumberOfElemens; ++ArrayValue)
+    {
+        Value.Add(MakeShared<FJsonValueNumber>(ArrayValue));
+    }
+
+    TSharedPtr<FJsonObject> Object = MakeShared<FJsonObject>();
+    Object->SetArrayField(Key, Value);
+
+    const TArray<TSharedPtr<FJsonValue>>* ParseValue;
+    TestTrue(TEXT("ParseArray"), FJsonUtility::ParseArray(Object, KeyInvalid, ParseValue).HasError());
+
+    return true;
+}
+
+IMPLEMENT_SIMPLE_AUTOMATION_TEST(TControlJsonUtilitiesCreateJsonFromData, "Plugin.CaptureProtocolStack.Control.ControlJsonUtilities.CreateJsonFromData.Success", EAutomationTestFlags::ApplicationContextMask | EAutomationTestFlags::SmokeFilter)
+bool TControlJsonUtilitiesCreateJsonFromData::RunTest(const FString& InParameters)
+{
+	auto ConvertedString = StringCast<UTF8CHAR>(TEXT("{\"Hello\":\"World\",\"Time\":123123123}"));
+	
+	TArray<uint8> Data(reinterpret_cast<const uint8*>(ConvertedString.Get()), ConvertedString.Length());
+
+	TSharedPtr<FJsonObject> Json = MakeShared<FJsonObject>();
+	TestTrue(TEXT("CreateJsonFromData"), FJsonUtility::CreateJsonFromUTF8Data(Data, Json));
+
+	FString ValueStr;
+	TestTrue(TEXT("CreateJsonFromData"), Json->TryGetStringField(TEXT("Hello"), ValueStr));
+	TestEqual(TEXT("CreateJsonFromData"), TEXT("World"), ValueStr);
+
+	uint32 ValueInt;
+	TestTrue(TEXT("CreateJsonFromData"), Json->TryGetNumberField(TEXT("Time"), ValueInt));
+	TestEqual(TEXT("CreateJsonFromData"), 123123123, ValueInt);
+
+	return true;
+}
+
+IMPLEMENT_SIMPLE_AUTOMATION_TEST(TControlJsonUtilitiesCreateJsonFromDataFail, "Plugin.CaptureProtocolStack.Control.ControlJsonUtilities.CreateJsonFromData.Failure", EAutomationTestFlags::ApplicationContextMask | EAutomationTestFlags::SmokeFilter)
+bool TControlJsonUtilitiesCreateJsonFromDataFail::RunTest(const FString& InParameters)
+{
+	auto ConvertedString = StringCast<UTF8CHAR>(TEXT("\"Hello\":\"World\",\"Time\":123123123"));
+
+	TArray<uint8> Data(reinterpret_cast<const uint8*>(ConvertedString.Get()), ConvertedString.Length());
+
+	TSharedPtr<FJsonObject> Json = MakeShared<FJsonObject>();
+	TestFalse(TEXT("CreateJsonFromData"), FJsonUtility::CreateJsonFromUTF8Data(Data, Json));
+
+	return true;
+}
+
+IMPLEMENT_SIMPLE_AUTOMATION_TEST(TControlJsonUtilitiesCreateDataFromJson, "Plugin.CaptureProtocolStack.Control.ControlJsonUtilities.CreateDataFromJson.Success", EAutomationTestFlags::ApplicationContextMask | EAutomationTestFlags::SmokeFilter)
+bool TControlJsonUtilitiesCreateDataFromJson::RunTest(const FString& InParameters)
+{
+	TSharedPtr<FJsonObject> Json = MakeShared<FJsonObject>();
+	Json->SetStringField(TEXT("Hello"), TEXT("World"));
+	Json->SetNumberField(TEXT("Time"), 123123123);
+
+	auto ConvertedString = StringCast<UTF8CHAR>(TEXT("{\"Hello\":\"World\",\"Time\":123123123}"));
+	TArray<uint8> ExpectedData(reinterpret_cast<const uint8*>(ConvertedString.Get()), ConvertedString.Length());
+
+	TArray<uint8> Data;
+	TestTrue(TEXT("CreateDataFromJson"), FJsonUtility::CreateUTF8DataFromJson(Json, Data));
+	TestEqual(TEXT("CreateDataFromJson"), Data, ExpectedData);
+
+	return true;
+}
+
+#endif // WITH_DEV_AUTOMATION_TESTS
\ No newline at end of file
diff --git a/Engine/Plugins/VirtualProduction/CaptureManager/Source/CaptureProtocolStack/Private/Tests/Control/TestControlMessage.cpp b/Engine/Plugins/VirtualProduction/CaptureManager/Source/CaptureProtocolStack/Private/Tests/Control/TestControlMessage.cpp
new file mode 100644
--- /dev/null
+++ b/Engine/Plugins/VirtualProduction/CaptureManager/Source/CaptureProtocolStack/Private/Tests/Control/TestControlMessage.cpp
@@ -0,0 +1,149 @@
+// Copyright Epic Games, Inc. All Rights Reserved.
+
+#include "Control/Messages/ControlMessage.h"
+
+#include "Control/Messages/ControlJsonUtilities.h"
+
+#include "Misc/AutomationTest.h"
+
+#if WITH_DEV_AUTOMATION_TESTS
+
+#define VERSION 1
+
+IMPLEMENT_SIMPLE_AUTOMATION_TEST(TControlMessageDeserializeTest, "Plugin.CaptureProtocolStack.Control.ControlMessage.Deserialize.Success", EAutomationTestFlags::ApplicationContextMask | EAutomationTestFlags::SmokeFilter)
+bool TControlMessageDeserializeTest::RunTest(const FString& InParameters)
+{
+    FString Payload = TEXT("{\"sessionId\":\"handshake\",\"transactionId\":123456789,\"timestamp\":112233445566,\"type\":\"request\",\"addressPath\":\"/session/start\"}");
+
+	TArray<uint8> Data;
+	Data.Append(Payload);
+
+    TProtocolResult<FControlMessage> DeserializeResult = FControlMessage::Deserialize(FControlPacket(VERSION, MoveTemp(Data)));
+    TestTrue(TEXT("Deserialize"), DeserializeResult.HasValue());
+
+    FControlMessage Message = DeserializeResult.StealValue();
+    TestEqual(TEXT("Deserialize"), Message.GetAddressPath(), TEXT("/session/start"));
+    TestEqual(TEXT("Deserialize"), Message.GetSessionId(), TEXT("handshake"));
+    TestEqual(TEXT("Deserialize"), Message.GetTransactionId(), (uint32) 123456789);
+    TestEqual(TEXT("Deserialize"), Message.GetTimestamp(), (uint64) 112233445566);
+    TestEqual(TEXT("Deserialize"), Message.GetType(), FControlMessage::EType::Request);
+    TestNull(TEXT("Deserialize"), Message.GetBody().Get());
+    TestTrue(TEXT("Deserialize"), Message.GetErrorName().IsEmpty());
+    TestTrue(TEXT("Deserialize"), Message.GetErrorDescription().IsEmpty());
+
+    return true;
+}
+
+IMPLEMENT_SIMPLE_AUTOMATION_TEST(TControlMessageDeserializeTestInvalidPayload, "Plugin.CaptureProtocolStack.Control.ControlMessage.DeserializeOne.InvalidPayload", EAutomationTestFlags::ApplicationContextMask | EAutomationTestFlags::SmokeFilter)
+bool TControlMessageDeserializeTestInvalidPayload::RunTest(const FString& InParameters)
+{
+	// Prepare test
+	FString Payload = TEXT("Hello");
+
+	TArray<uint8> Data;
+	Data.Append(Payload);
+
+	TProtocolResult<FControlMessage> DeserializeResult = FControlMessage::Deserialize(FControlPacket(VERSION, MoveTemp(Data)));
+	TestTrue(TEXT("Deserialize"), DeserializeResult.HasError());
+
+	return true;
+}
+
+IMPLEMENT_SIMPLE_AUTOMATION_TEST(TControlMessageDeserializeSessionIdMissingTest, "Plugin.CaptureProtocolStack.Control.ControlMessage.Deserialize.SessionIdMissing", EAutomationTestFlags::ApplicationContextMask | EAutomationTestFlags::SmokeFilter)
+bool TControlMessageDeserializeSessionIdMissingTest::RunTest(const FString& InParameters)
+{
+    FString Payload = TEXT("{\"transactionId\":123456789,\"timestamp\":112233445566,\"type\":\"request\",\"addressPath\":\"/session/start\"}");
+
+	TArray<uint8> Data;
+	Data.Append(Payload);
+
+	TProtocolResult<FControlMessage> DeserializeResult = FControlMessage::Deserialize(FControlPacket(VERSION, MoveTemp(Data)));
+    TestTrue(TEXT("Deserialize"), DeserializeResult.HasError());
+
+    return true;
+}
+
+IMPLEMENT_SIMPLE_AUTOMATION_TEST(TControlMessageDeserializeTransactionIdMissingTest, "Plugin.CaptureProtocolStack.Control.ControlMessage.Deserialize.TransactionIdMissing", EAutomationTestFlags::ApplicationContextMask | EAutomationTestFlags::SmokeFilter)
+bool TControlMessageDeserializeTransactionIdMissingTest::RunTest(const FString& InParameters)
+{
+    FString Payload = TEXT("{\"sessionId\":\"handshake\",\"timestamp\":112233445566,\"type\":\"request\",\"addressPath\":\"/session/start\"}");
+
+	TArray<uint8> Data;
+	Data.Append(Payload);
+
+	TProtocolResult<FControlMessage> DeserializeResult = FControlMessage::Deserialize(FControlPacket(VERSION, MoveTemp(Data)));
+    TestTrue(TEXT("Deserialize"), DeserializeResult.HasError());
+
+    return true;
+}
+
+IMPLEMENT_SIMPLE_AUTOMATION_TEST(TControlMessageDeserializeTimestampMissingTest, "Plugin.CaptureProtocolStack.Control.ControlMessage.Deserialize.TimestampMissing", EAutomationTestFlags::ApplicationContextMask | EAutomationTestFlags::SmokeFilter)
+bool TControlMessageDeserializeTimestampMissingTest::RunTest(const FString& InParameters)
+{
+    FString Payload = TEXT("{\"sessionId\":\"handshake\",\"transactionId\":123456789,\"type\":\"request\",\"addressPath\":\"/session/start\"}");
+
+	TArray<uint8> Data;
+	Data.Append(Payload);
+
+	TProtocolResult<FControlMessage> DeserializeResult = FControlMessage::Deserialize(FControlPacket(VERSION, MoveTemp(Data)));
+    TestTrue(TEXT("Deserialize"), DeserializeResult.HasError());
+
+    return true;
+}
+
+IMPLEMENT_SIMPLE_AUTOMATION_TEST(TControlMessageDeserializeTypeMissingTest, "Plugin.CaptureProtocolStack.Control.ControlMessage.Deserialize.TypeMissing", EAutomationTestFlags::ApplicationContextMask | EAutomationTestFlags::SmokeFilter)
+bool TControlMessageDeserializeTypeMissingTest::RunTest(const FString& InParameters)
+{
+    FString Payload = TEXT("{\"sessionId\":\"handshake\",\"transactionId\":123456789,\"timestamp\":112233445566,\"addressPath\":\"/session/start\"}");
+
+	TArray<uint8> Data;
+	Data.Append(Payload);
+
+	TProtocolResult<FControlMessage> DeserializeResult = FControlMessage::Deserialize(FControlPacket(VERSION, MoveTemp(Data)));
+    TestTrue(TEXT("Deserialize"), DeserializeResult.HasError());
+
+    return true;
+}
+
+IMPLEMENT_SIMPLE_AUTOMATION_TEST(TControlMessageDeserializeAddressPathMissingTest, "Plugin.CaptureProtocolStack.Control.ControlMessage.Deserialize.AddressPathMissing", EAutomationTestFlags::ApplicationContextMask | EAutomationTestFlags::SmokeFilter)
+bool TControlMessageDeserializeAddressPathMissingTest::RunTest(const FString& InParameters)
+{
+    FString Payload = TEXT("{\"sessionId\":\"handshake\",\"transactionId\":123456789,\"timestamp\":112233445566,\"type\":\"request\"}");
+
+	TArray<uint8> Data;
+	Data.Append(Payload);
+
+	TProtocolResult<FControlMessage> DeserializeResult = FControlMessage::Deserialize(FControlPacket(VERSION, MoveTemp(Data)));
+    TestTrue(TEXT("Deserialize"), DeserializeResult.HasError());
+
+    return true;
+}
+
+IMPLEMENT_SIMPLE_AUTOMATION_TEST(TControlMessageSerializeTest, "Plugin.CaptureProtocolStack.Control.ControlMessage.Serialize.Success", EAutomationTestFlags::ApplicationContextMask | EAutomationTestFlags::SmokeFilter)
+bool TControlMessageSerializeTest::RunTest(const FString& InParameters)
+{
+    FControlMessage Message(TEXT("/session/start"), FControlMessage::EType::Request, nullptr);
+    Message.SetSessionId(TEXT("handshake"));
+    Message.SetTransactionId(123456789);
+    Message.SetTimestamp(112233445566);
+
+    TProtocolResult<FControlPacket> SerializeResult = FControlMessage::Serialize(Message);
+    TestTrue(TEXT("Serialize"), SerializeResult.HasValue());
+
+    TProtocolResult<FControlMessage> DeserializeResult = FControlMessage::Deserialize(SerializeResult.StealValue());
+    TestTrue(TEXT("Serialize"), DeserializeResult.HasValue());
+
+    FControlMessage DeserializedMessage = DeserializeResult.StealValue();
+    TestEqual(TEXT("Serialize"), Message.GetAddressPath(), DeserializedMessage.GetAddressPath());
+    TestEqual(TEXT("Serialize"), Message.GetSessionId(), DeserializedMessage.GetSessionId());
+    TestEqual(TEXT("Serialize"), Message.GetTransactionId(), DeserializedMessage.GetTransactionId());
+    TestEqual(TEXT("Serialize"), Message.GetTimestamp(), DeserializedMessage.GetTimestamp());
+    TestEqual(TEXT("Serialize"), Message.GetType(), DeserializedMessage.GetType());
+    TestEqual(TEXT("Serialize"), Message.GetBody(), DeserializedMessage.GetBody());
+    TestEqual(TEXT("Serialize"), Message.GetErrorName(), DeserializedMessage.GetErrorName());
+    TestEqual(TEXT("Serialize"), Message.GetErrorDescription(), DeserializedMessage.GetErrorDescription());
+
+    return true;
+}
+
+#endif //WITH_DEV_AUTOMATION_TESTS
\ No newline at end of file
diff --git a/Engine/Plugins/VirtualProduction/CaptureManager/Source/CaptureProtocolStack/Private/Tests/Control/TestControlPacket.cpp b/Engine/Plugins/VirtualProduction/CaptureManager/Source/CaptureProtocolStack/Private/Tests/Control/TestControlPacket.cpp
new file mode 100644
--- /dev/null
+++ b/Engine/Plugins/VirtualProduction/CaptureManager/Source/CaptureProtocolStack/Private/Tests/Control/TestControlPacket.cpp
@@ -0,0 +1,192 @@
+// Copyright Epic Games, Inc. All Rights Reserved.
+
+#include "Control/Communication/ControlPacket.h"
+
+#include "Control/Messages/ControlJsonUtilities.h"
+
+#include "Misc/AutomationTest.h"
+
+#include "Tests/Utility.h"
+
+#if WITH_DEV_AUTOMATION_TESTS
+
+IMPLEMENT_SIMPLE_AUTOMATION_TEST(TControlPacketDeserializeOneTest, "Plugin.CaptureProtocolStack.Control.ControlPacket.DeserializeOne.Success", EAutomationTestFlags::ApplicationContextMask | EAutomationTestFlags::SmokeFilter)
+bool TControlPacketDeserializeOneTest::RunTest(const FString& InParameters)
+{
+    // Prepare test
+    const TArray<uint8> Header = { 'C', 'P', 'S', 'C', 'O', 'N', 'T', 'R', 'O', 'L', '\0'};
+	const uint16 Version = 1;
+    const FString Payload = TEXT("{\"Hello\": \"World\"}");
+    const uint32 Length = Payload.Len();
+
+    TArray<uint8> Packet;
+    Packet.Reserve(Header.Num() + sizeof(Version) + sizeof(Length) + Length);
+    Packet.Append(Header);
+	Packet.Append(reinterpret_cast<const uint8*>(&Version), sizeof(Version));
+    Packet.Append(reinterpret_cast<const uint8*>(&Length), sizeof(Length));
+    Packet.Append(Payload);
+
+    FDataProvider DataProvider(MoveTemp(Packet));
+
+    // Deserialization
+    TProtocolResult<FControlPacketHeader> HeaderResult = FControlPacketHeader::Deserialize(DataProvider);
+    TestTrue(TEXT("Header Deserialize"), HeaderResult.HasValue());
+
+    FControlPacketHeader ControlHeader = HeaderResult.StealValue();
+    TProtocolResult<FControlPacket> PacketProtocolResult = FControlPacket::Deserialize(ControlHeader, DataProvider);
+    TestTrue(TEXT("Packet Deserialize"), PacketProtocolResult.HasValue());
+
+    FControlPacket ControlPacket = PacketProtocolResult.StealValue();
+	TestEqual(TEXT("Packet Deserialize"), ControlPacket.GetVersion(), Version);
+    TestEqual(TEXT("Packet Deserialize"), ControlPacket.GetPayloadSize(), Length);
+
+    TArray<uint8> PayloadData = ControlPacket.GetPayload();
+
+	TSharedPtr<FJsonObject> JsonPayload;
+	bool Result = FJsonUtility::CreateJsonFromUTF8Data(PayloadData, JsonPayload);
+	TestTrue(TEXT("Packet Deserialize"), Result);
+
+    FString Field;
+    TestTrue(TEXT("Packet Deserialize"), JsonPayload->TryGetStringField(TEXT("Hello"), Field));
+    TestEqual(TEXT("Packet Deserialize"), TEXT("World"), Field);
+
+    return true;
+}
+
+IMPLEMENT_SIMPLE_AUTOMATION_TEST(TControlPacketDeserializeOneTestInvalidHeader, "Plugin.CaptureProtocolStack.Control.ControlPacket.DeserializeOne.InvalidHeader", EAutomationTestFlags::ApplicationContextMask | EAutomationTestFlags::SmokeFilter)
+bool TControlPacketDeserializeOneTestInvalidHeader::RunTest(const FString& InParameters)
+{
+    // Prepare test
+    const TArray<uint8> Header = { 'A', 'A', 'A', '\0' };
+	const uint16 Version = 1;
+    const FString Payload = TEXT("{\"Hello\": \"World\"}");
+    const uint32 Length = Payload.Len();
+
+    TArray<uint8> Packet;
+    Packet.Reserve(Header.Num() + sizeof(Version) + sizeof(Length) + Length);
+    Packet.Append(Header);
+	Packet.Append(reinterpret_cast<const uint8*>(&Version), sizeof(Version));
+	Packet.Append(reinterpret_cast<const uint8*>(&Length), sizeof(Length));
+    Packet.Append(Payload);
+
+    FDataProvider DataProvider(MoveTemp(Packet));
+
+    // Deserialization
+    TProtocolResult<FControlPacketHeader> HeaderResult = FControlPacketHeader::Deserialize(DataProvider);
+    TestTrue(TEXT("Header Deserialize"), HeaderResult.HasError());
+
+    return true;
+}
+
+IMPLEMENT_SIMPLE_AUTOMATION_TEST(TControlPacketDeserializeMoreTest, "Plugin.CaptureProtocolStack.Control.ControlPacket.DeserializeMore.Success", EAutomationTestFlags::ApplicationContextMask | EAutomationTestFlags::SmokeFilter)
+bool TControlPacketDeserializeMoreTest::RunTest(const FString& InParameters)
+{
+    // Prepare test
+    const TArray<uint8> Header = { 'C', 'P', 'S', 'C', 'O', 'N', 'T', 'R', 'O', 'L', '\0' };
+	const uint16 Version = 1;
+    const FString Payload = TEXT("{\"Hello\": \"World\"}");
+    const uint32 Length = Payload.Len();
+
+    constexpr uint32 NumberOfPackets = 5;
+
+    TArray<uint8> Packet;
+    for (int32 PacketIndex = 0; PacketIndex < NumberOfPackets; ++PacketIndex)
+    {
+        Packet.Reserve(Header.Num() + sizeof(Version) + sizeof(Length) + Length);
+        Packet.Append(Header);
+		Packet.Append(reinterpret_cast<const uint8*>(&Version), sizeof(Version));
+		Packet.Append(reinterpret_cast<const uint8*>(&Length), sizeof(Length));
+        Packet.Append(Payload);
+    }
+
+    FDataProvider DataProvider(MoveTemp(Packet));
+
+    // Deserialization
+    uint32 PacketsDeserialized = 0;
+    while (PacketsDeserialized != NumberOfPackets)
+    {
+        TProtocolResult<FControlPacketHeader> HeaderResult = FControlPacketHeader::Deserialize(DataProvider);
+        TestTrue(TEXT("Header Deserialize"), HeaderResult.HasValue());
+
+        FControlPacketHeader ControlHeader = HeaderResult.StealValue();
+        TProtocolResult<FControlPacket> PacketProtocolResult = FControlPacket::Deserialize(ControlHeader, DataProvider);
+        TestTrue(TEXT("Packet Deserialize"), PacketProtocolResult.HasValue());
+
+        FControlPacket ControlPacket = PacketProtocolResult.StealValue();
+		TestEqual(TEXT("Packet Deserialize"), ControlPacket.GetVersion(), Version);
+        TestEqual(TEXT("Packet Deserialize"), ControlPacket.GetPayloadSize(), Length);
+
+		TArray<uint8> PayloadData = ControlPacket.GetPayload();
+
+		TSharedPtr<FJsonObject> JsonPayload;
+		bool Result = FJsonUtility::CreateJsonFromUTF8Data(PayloadData, JsonPayload);
+		TestTrue(TEXT("Packet Deserialize"), Result);
+
+        FString Field;
+        TestTrue(TEXT("Packet Deserialize"), JsonPayload->TryGetStringField(TEXT("Hello"), Field));
+        TestEqual(TEXT("Packet Deserialize"), TEXT("World"), Field);
+
+        ++PacketsDeserialized;
+    }
+
+    return true;
+}
+
+IMPLEMENT_SIMPLE_AUTOMATION_TEST(TControlPacketSerializeOneTest, "Plugin.CaptureProtocolStack.Control.ControlPacket.SerializeOne.Success", EAutomationTestFlags::ApplicationContextMask | EAutomationTestFlags::SmokeFilter)
+bool TControlPacketSerializeOneTest::RunTest(const FString& InParameters)
+{
+	FDataSender DataReceiver;
+
+	// Prepare test
+	const TArray<uint8> Header = { 'C', 'P', 'S', 'C', 'O', 'N', 'T', 'R', 'O', 'L', '\0' };
+	const uint16 Version = 1;
+	const FString Payload = TEXT("{\"Hello\":\"World\"}");
+	const uint32 Length = Payload.Len();
+
+	TArray<uint8> PayloadData;
+	PayloadData.Append(Payload);
+
+	TArray<uint8> Packet;
+	Packet.Reserve(Header.Num() + sizeof(Version) + sizeof(Length) + Length);
+	Packet.Append(Header);
+	Packet.Append(reinterpret_cast<const uint8*>(&Version), sizeof(Version));
+	Packet.Append(reinterpret_cast<const uint8*>(&Length), sizeof(Length));
+	Packet.Append(PayloadData);
+
+	FControlPacket ControlPacket(Version, MoveTemp(PayloadData));
+
+	TestTrue(TEXT("Packet Serialize"), FControlPacket::Serialize(ControlPacket, DataReceiver).HasValue());
+	TestEqual(TEXT("Packet Serialize"), Packet, DataReceiver.GetData());
+
+	return true;
+}
+
+IMPLEMENT_SIMPLE_AUTOMATION_TEST(TControlPacketSerializeOneTestFailed, "Plugin.CaptureProtocolStack.Control.ControlPacket.SerializeOne.Failure", EAutomationTestFlags::ApplicationContextMask | EAutomationTestFlags::SmokeFilter)
+bool TControlPacketSerializeOneTestFailed::RunTest(const FString& InParameters)
+{
+	FFailedDataSender DataReceiver;
+
+	// Prepare test
+	const TArray<uint8> Header = { 'C', 'P', 'S', 'C', 'O', 'N', 'T', 'R', 'O', 'L', '\0' };
+	const uint16 Version = 1;
+	const FString Payload = TEXT("{\"Hello\":\"World\"}");
+	const uint32 Length = Payload.Len();
+
+	TArray<uint8> PayloadData;
+	PayloadData.Append(Payload);
+
+	TArray<uint8> Packet;
+	Packet.Reserve(Header.Num() + sizeof(Version) + sizeof(Length) + Length);
+	Packet.Append(Header);
+	Packet.Append(reinterpret_cast<const uint8*>(&Version), sizeof(Version));
+	Packet.Append(reinterpret_cast<const uint8*>(&Length), sizeof(Length));
+	Packet.Append(PayloadData);
+
+	FControlPacket ControlPacket(Version, MoveTemp(PayloadData));
+
+	TestTrue(TEXT("Packet Serialize"), FControlPacket::Serialize(ControlPacket, DataReceiver).HasError());
+
+	return true;
+}
+
+#endif //WITH_DEV_AUTOMATION_TESTS
\ No newline at end of file
diff --git a/Engine/Plugins/VirtualProduction/CaptureManager/Source/CaptureProtocolStack/Private/Tests/Control/TestControlRequest.cpp b/Engine/Plugins/VirtualProduction/CaptureManager/Source/CaptureProtocolStack/Private/Tests/Control/TestControlRequest.cpp
new file mode 100644
--- /dev/null
+++ b/Engine/Plugins/VirtualProduction/CaptureManager/Source/CaptureProtocolStack/Private/Tests/Control/TestControlRequest.cpp
@@ -0,0 +1,182 @@
+// Copyright Epic Games, Inc. All Rights Reserved.
+
+#include "Control/Messages/ControlRequest.h"
+
+#include "Control/Messages/ControlJsonUtilities.h"
+#include "Control/Messages/Constants.h"
+
+#include "Misc/AutomationTest.h"
+
+#if WITH_DEV_AUTOMATION_TESTS
+
+IMPLEMENT_SIMPLE_AUTOMATION_TEST(TGetServerInformationTest, "Plugin.CaptureProtocolStack.Control.ControlRequest.GetServerInformation.Success", EAutomationTestFlags::ApplicationContextMask | EAutomationTestFlags::SmokeFilter)
+bool TGetServerInformationTest::RunTest(const FString& InParameters)
+{
+	FGetServerInformationRequest Request;
+
+	TestEqual(TEXT("GetAddresPath"), Request.GetAddressPath(), UE::CPS::AddressPaths::GGetServerInformation);
+	TestFalse(TEXT("GetBody"), Request.GetBody().IsValid()); // Request doesn't have a Body
+
+	return true;
+}
+
+IMPLEMENT_SIMPLE_AUTOMATION_TEST(TKeepAliveRequestTest, "Plugin.CaptureProtocolStack.Control.ControlRequest.KeepAlive.Success", EAutomationTestFlags::ApplicationContextMask | EAutomationTestFlags::SmokeFilter)
+bool TKeepAliveRequestTest::RunTest(const FString& InParameters)
+{        
+    FKeepAliveRequest Request;
+
+    TestEqual(TEXT("GetAddresPath"), Request.GetAddressPath(), UE::CPS::AddressPaths::GKeepAlive);
+    TestFalse(TEXT("GetBody"), Request.GetBody().IsValid()); // Request doesn't have a Body
+
+    return true;
+}
+
+IMPLEMENT_SIMPLE_AUTOMATION_TEST(TStartSessionRequestTest, "Plugin.CaptureProtocolStack.Control.ControlRequest.StartSession.Success", EAutomationTestFlags::ApplicationContextMask | EAutomationTestFlags::SmokeFilter)
+bool TStartSessionRequestTest::RunTest(const FString& InParameters)
+{
+    FStartSessionRequest Request;
+
+    TestEqual(TEXT("GetAddresPath"), Request.GetAddressPath(), UE::CPS::AddressPaths::GStartSession);
+    TestFalse(TEXT("GetBody"), Request.GetBody().IsValid()); // Request doesn't have a Body
+
+    return true;
+}
+
+IMPLEMENT_SIMPLE_AUTOMATION_TEST(TStopSessionRequestTest, "Plugin.CaptureProtocolStack.Control.ControlRequest.StopSession.Success", EAutomationTestFlags::ApplicationContextMask | EAutomationTestFlags::SmokeFilter)
+bool TStopSessionRequestTest::RunTest(const FString& InParameters)
+{
+    FStopSessionRequest Request;
+
+    TestEqual(TEXT("GetAddresPath"), Request.GetAddressPath(), UE::CPS::AddressPaths::GStopSession);
+    TestFalse(TEXT("GetBody"), Request.GetBody().IsValid()); // Request doesn't have a Body
+
+    return true;
+}
+
+IMPLEMENT_SIMPLE_AUTOMATION_TEST(TSubscribeRequestTest, "Plugin.CaptureProtocolStack.Control.ControlRequest.Subscribe.Success", EAutomationTestFlags::ApplicationContextMask | EAutomationTestFlags::SmokeFilter)
+bool TSubscribeRequestTest::RunTest(const FString& InParameters)
+{
+    FSubscribeRequest Request;
+
+    TestEqual(TEXT("GetAddresPath"), Request.GetAddressPath(), UE::CPS::AddressPaths::GSubscribe);
+    TestFalse(TEXT("GetBody"), Request.GetBody().IsValid()); // Request doesn't have a Body
+
+    return true;
+}
+
+IMPLEMENT_SIMPLE_AUTOMATION_TEST(TUnsubscribeRequestTest, "Plugin.CaptureProtocolStack.Control.ControlRequest.Unsubscribe.Success", EAutomationTestFlags::ApplicationContextMask | EAutomationTestFlags::SmokeFilter)
+bool TUnsubscribeRequestTest::RunTest(const FString& InParameters)
+{
+    FUnsubscribeRequest Request;
+
+    TestEqual(TEXT("GetAddresPath"), Request.GetAddressPath(), UE::CPS::AddressPaths::GUnsubscribe);
+    TestFalse(TEXT("GetBody"), Request.GetBody().IsValid()); // Request doesn't have a Body
+
+    return true;
+}
+
+IMPLEMENT_SIMPLE_AUTOMATION_TEST(TGetStateRequestTest, "Plugin.CaptureProtocolStack.Control.ControlRequest.GetState.Success", EAutomationTestFlags::ApplicationContextMask | EAutomationTestFlags::SmokeFilter)
+bool TGetStateRequestTest::RunTest(const FString& InParameters)
+{
+    FGetStateRequest Request;
+
+    TestEqual(TEXT("GetAddresPath"), Request.GetAddressPath(), UE::CPS::AddressPaths::GGetState);
+    TestFalse(TEXT("GetBody"), Request.GetBody().IsValid()); // Request doesn't have a Body
+
+    return true;
+}
+
+IMPLEMENT_SIMPLE_AUTOMATION_TEST(TStartRecordingTakeRequestTest, "Plugin.CaptureProtocolStack.Control.ControlRequest.StartRecordingTake.Success", EAutomationTestFlags::ApplicationContextMask | EAutomationTestFlags::SmokeFilter)
+bool TStartRecordingTakeRequestTest::RunTest(const FString& InParameters)
+{
+	FString String = TEXT("{\
+        \"slateName\": \"Slate\",\
+        \"takeNumber\": 0,\
+        \"subject\": \"Subject\",\
+        \"scenario\": \"Scenario\",\
+        \"tags\": [\"Tag1\", \"Tag2\", \"Tag3\"]\
+    }");
+
+	String.ConvertTabsToSpacesInline(2);
+	String.RemoveSpacesInline();
+
+	auto ConvertedString = StringCast<UTF8CHAR>(*String);
+
+    FString SlateName = TEXT("Slate");
+    FString Subject = TEXT("Subject");
+    FString Scenario = TEXT("Scenario");
+    TArray<FString> Tags = { TEXT("Tag1"), TEXT("Tag2"), TEXT("Tag3") };
+    
+    FStartRecordingTakeRequest Request(SlateName, 0, Subject, Scenario, Tags);
+
+    TArray<uint8> Body;
+    TestTrue(TEXT("CreateDataFromJson"), FJsonUtility::CreateUTF8DataFromJson(Request.GetBody(), Body));
+
+    TestEqual(TEXT("GetAddresPath"), Request.GetAddressPath(), UE::CPS::AddressPaths::GStartRecordingTake);
+
+	TArray<uint8> BodyExpected(reinterpret_cast<const uint8*>(ConvertedString.Get()), ConvertedString.Length());
+	TestEqual(TEXT("GetBody"), Body, BodyExpected);
+
+    return true;
+}
+
+IMPLEMENT_SIMPLE_AUTOMATION_TEST(TStopRecordingTakeRequestTest, "Plugin.CaptureProtocolStack.Control.ControlRequest.StopRecordingTake.Success", EAutomationTestFlags::ApplicationContextMask | EAutomationTestFlags::SmokeFilter)
+bool TStopRecordingTakeRequestTest::RunTest(const FString& InParameters)
+{
+	FStopRecordingTakeRequest Request;
+
+	TestEqual(TEXT("GetAddresPath"), Request.GetAddressPath(), UE::CPS::AddressPaths::GStopRecordingTake);
+	TestFalse(TEXT("GetBody"), Request.GetBody().IsValid()); // Request doesn't have a Body
+
+	return true;
+}
+
+IMPLEMENT_SIMPLE_AUTOMATION_TEST(TAbortRecordingTakeRequestTest, "Plugin.CaptureProtocolStack.Control.ControlRequest.AbortRecordingTake.Success", EAutomationTestFlags::ApplicationContextMask | EAutomationTestFlags::SmokeFilter)
+bool TAbortRecordingTakeRequestTest::RunTest(const FString& InParameters)
+{
+	FAbortRecordingTakeRequest Request;
+
+	TestEqual(TEXT("GetAddresPath"), Request.GetAddressPath(), UE::CPS::AddressPaths::GAbortRecordingTake);
+	TestFalse(TEXT("GetBody"), Request.GetBody().IsValid()); // Request doesn't have a Body
+
+	return true;
+}
+
+IMPLEMENT_SIMPLE_AUTOMATION_TEST(TGetTakeListRequestTest, "Plugin.CaptureProtocolStack.Control.ControlRequest.GetTakeList.Success", EAutomationTestFlags::ApplicationContextMask | EAutomationTestFlags::SmokeFilter)
+bool TGetTakeListRequestTest::RunTest(const FString& InParameters)
+{
+    FGetTakeListRequest Request;
+
+    TestEqual(TEXT("GetAddresPath"), Request.GetAddressPath(), UE::CPS::AddressPaths::GGetTakeList);
+    TestFalse(TEXT("GetBody"), Request.GetBody().IsValid()); // Request doesn't have a Body
+
+    return true;
+}
+
+IMPLEMENT_SIMPLE_AUTOMATION_TEST(TGetTakeMetadataRequestTest, "Plugin.CaptureProtocolStack.Control.ControlRequest.GetTakeMetadata.Success", EAutomationTestFlags::ApplicationContextMask | EAutomationTestFlags::SmokeFilter)
+bool TGetTakeMetadataRequestTest::RunTest(const FString& InParameters)
+{
+	FString String = TEXT("{\
+        \"names\": [\"TakeName1\", \"TakeName2\", \"TakeName3\"]\
+    }");
+
+	String.ConvertTabsToSpacesInline(2);
+	String.RemoveSpacesInline();
+
+    auto ConvertedString = StringCast<UTF8CHAR>(*String);
+
+    TArray<FString> Takes = { TEXT("TakeName1"), TEXT("TakeName2"), TEXT("TakeName3") };
+    FGetTakeMetadataRequest Request(Takes);
+
+    TArray<uint8> Body;
+    TestTrue(TEXT("CreateDataFromJson"), FJsonUtility::CreateUTF8DataFromJson(Request.GetBody(), Body));
+
+    TestEqual(TEXT("GetAddresPath"), Request.GetAddressPath(), UE::CPS::AddressPaths::GGetTakeMetadata);
+
+	TArray<uint8> BodyExpected(reinterpret_cast<const uint8*>(ConvertedString.Get()), ConvertedString.Length());
+    TestEqual(TEXT("GetBody"), Body, BodyExpected);
+
+    return true;
+}
+
+#endif // WITH_DEV_AUTOMATION_TESTS
\ No newline at end of file
diff --git a/Engine/Plugins/VirtualProduction/CaptureManager/Source/CaptureProtocolStack/Private/Tests/Control/TestControlResponse.cpp b/Engine/Plugins/VirtualProduction/CaptureManager/Source/CaptureProtocolStack/Private/Tests/Control/TestControlResponse.cpp
new file mode 100644
--- /dev/null
+++ b/Engine/Plugins/VirtualProduction/CaptureManager/Source/CaptureProtocolStack/Private/Tests/Control/TestControlResponse.cpp
@@ -0,0 +1,313 @@
+// Copyright Epic Games, Inc. All Rights Reserved.
+
+#include "Control/Messages/ControlResponse.h"
+
+#include "Control/Messages/ControlJsonUtilities.h"
+#include "Control/Messages/Constants.h"
+
+#include "Misc/AutomationTest.h"
+#include "Misc/DateTime.h"
+
+#ifdef WITH_DEV_AUTOMATION_TESTS
+
+IMPLEMENT_SIMPLE_AUTOMATION_TEST(TKeepAliveResponseTest, "Plugin.CaptureProtocolStack.Control.ControlResponse.KeepAlive.Success", EAutomationTestFlags::ApplicationContextMask | EAutomationTestFlags::SmokeFilter)
+bool TKeepAliveResponseTest::RunTest(const FString& InParameters)
+{
+    TSharedPtr<FJsonObject> Body;
+
+    FKeepAliveResponse Response;
+
+    TestEqual(TEXT("GetAddresPath"), Response.GetAddressPath(), UE::CPS::AddressPaths::GKeepAlive);
+    TestTrue(TEXT("Parse"), Response.Parse(Body).HasValue());
+
+    return true;
+}
+
+IMPLEMENT_SIMPLE_AUTOMATION_TEST(TStartSessionResponseTest, "Plugin.CaptureProtocolStack.Control.ControlResponse.StartSession.Success", EAutomationTestFlags::ApplicationContextMask | EAutomationTestFlags::SmokeFilter)
+bool TStartSessionResponseTest::RunTest(const FString& InParameters)
+{
+    TSharedPtr<FJsonObject> Body = MakeShared<FJsonObject>();
+    Body->SetStringField(UE::CPS::Properties::GSessionId, TEXT("SessionId"));
+
+    FStartSessionResponse Response;
+
+    TestEqual(TEXT("GetAddresPath"), Response.GetAddressPath(), UE::CPS::AddressPaths::GStartSession);
+    TestTrue(TEXT("Parse"), Response.Parse(Body).HasValue());
+
+    TestEqual(TEXT("Parse"), Response.GetSessionId(), TEXT("SessionId"));
+
+    return true;
+}
+
+IMPLEMENT_SIMPLE_AUTOMATION_TEST(TStopSessionResponseTest, "Plugin.CaptureProtocolStack.Control.ControlResponse.StopSession.Success", EAutomationTestFlags::ApplicationContextMask | EAutomationTestFlags::SmokeFilter)
+bool TStopSessionResponseTest::RunTest(const FString& InParameters)
+{
+    TSharedPtr<FJsonObject> Body;
+
+    FStopSessionResponse Response;
+
+    TestEqual(TEXT("GetAddresPath"), Response.GetAddressPath(), UE::CPS::AddressPaths::GStopSession);
+    TestTrue(TEXT("Parse"), Response.Parse(Body).HasValue());
+
+    return true;
+}
+
+IMPLEMENT_SIMPLE_AUTOMATION_TEST(TGetServerInformationResponseTest, "Plugin.CaptureProtocolStack.Control.ControlResponse.GetServerInformation.Success", EAutomationTestFlags::ApplicationContextMask | EAutomationTestFlags::SmokeFilter)
+bool TGetServerInformationResponseTest::RunTest(const FString& InParameters)
+{
+    TSharedPtr<FJsonObject> Body = MakeShared<FJsonObject>();
+    Body->SetStringField(UE::CPS::Properties::GId, TEXT("Id"));
+    Body->SetStringField(UE::CPS::Properties::GName, TEXT("Name"));
+    Body->SetStringField(UE::CPS::Properties::GModel, TEXT("Model"));
+    Body->SetStringField(UE::CPS::Properties::GPlatformName, TEXT("PlatformName"));
+    Body->SetStringField(UE::CPS::Properties::GPlatformVersion, TEXT("PlatformVersion"));
+    Body->SetStringField(UE::CPS::Properties::GSoftwareName, TEXT("SoftwareName"));
+    Body->SetStringField(UE::CPS::Properties::GSoftwareVersion, TEXT("SoftwareVersion"));
+	Body->SetNumberField(UE::CPS::Properties::GExportPort, 12345);
+
+    FGetServerInformationResponse Response;
+
+    TestEqual(TEXT("GetAddresPath"), Response.GetAddressPath(), UE::CPS::AddressPaths::GGetServerInformation);
+    TestTrue(TEXT("Parse"), Response.Parse(Body).HasValue());
+
+    TestEqual(TEXT("Parse"), Response.GetId(), TEXT("Id"));
+    TestEqual(TEXT("Parse"), Response.GetName(), TEXT("Name"));
+    TestEqual(TEXT("Parse"), Response.GetModel(), TEXT("Model"));
+    TestEqual(TEXT("Parse"), Response.GetPlatformName(), TEXT("PlatformName"));
+    TestEqual(TEXT("Parse"), Response.GetPlatformVersion(), TEXT("PlatformVersion"));
+    TestEqual(TEXT("Parse"), Response.GetSoftwareName(), TEXT("SoftwareName"));
+    TestEqual(TEXT("Parse"), Response.GetSoftwareVersion(), TEXT("SoftwareVersion"));
+	TestEqual(TEXT("Parse"), Response.GetExportPort(), 12345);
+
+    return true;
+}
+
+IMPLEMENT_SIMPLE_AUTOMATION_TEST(TSubscribeResponseTest, "Plugin.CaptureProtocolStack.Control.ControlResponse.Subscribe.Success", EAutomationTestFlags::ApplicationContextMask | EAutomationTestFlags::SmokeFilter)
+bool TSubscribeResponseTest::RunTest(const FString& InParameters)
+{
+    TSharedPtr<FJsonObject> Body;
+
+    FSubscribeResponse Response;
+
+    TestEqual(TEXT("GetAddresPath"), Response.GetAddressPath(), UE::CPS::AddressPaths::GSubscribe);
+    TestTrue(TEXT("Parse"), Response.Parse(Body).HasValue());
+
+    return true;
+}
+
+IMPLEMENT_SIMPLE_AUTOMATION_TEST(TUnsubscribeResponseTest, "Plugin.CaptureProtocolStack.Control.ControlResponse.Unsubscribe.Success", EAutomationTestFlags::ApplicationContextMask | EAutomationTestFlags::SmokeFilter)
+bool TUnsubscribeResponseTest::RunTest(const FString& InParameters)
+{
+    TSharedPtr<FJsonObject> Body;
+
+    FUnsubscribeResponse Response;
+
+    TestEqual(TEXT("GetAddresPath"), Response.GetAddressPath(), UE::CPS::AddressPaths::GUnsubscribe);
+    TestTrue(TEXT("Parse"), Response.Parse(Body).HasValue());
+
+    return true;
+}
+
+IMPLEMENT_SIMPLE_AUTOMATION_TEST(TGetStateResponseTest, "Plugin.CaptureProtocolStack.Control.ControlResponse.GetState.Success", EAutomationTestFlags::ApplicationContextMask | EAutomationTestFlags::SmokeFilter)
+bool TGetStateResponseTest::RunTest(const FString& InParameters)
+{
+	TSharedPtr<FJsonObject> Body = MakeShared<FJsonObject>();
+	Body->SetBoolField(UE::CPS::Properties::GIsRecording, true);
+
+	TSharedPtr<FJsonObject> IOSSpecific = MakeShared<FJsonObject>();
+	IOSSpecific->SetNumberField(UE::CPS::Properties::GTotalCapacity, 100);
+	IOSSpecific->SetNumberField(UE::CPS::Properties::GRemainingCapacity, 100);
+	IOSSpecific->SetNumberField(UE::CPS::Properties::GBatteryLevel, 100.0);
+	IOSSpecific->SetStringField(UE::CPS::Properties::GThermalState, UE::CPS::Properties::GNominal);
+
+	Body->SetObjectField(UE::CPS::Properties::GPlatformState, IOSSpecific);
+
+	FGetStateResponse Response;
+
+	TestEqual(TEXT("GetAddresPath"), Response.GetAddressPath(), UE::CPS::AddressPaths::GGetState);
+	TestTrue(TEXT("Parse"), Response.Parse(Body).HasValue());
+
+	TestEqual(TEXT("Parse"), Response.IsRecording(), true);
+
+	FJsonValueObject PlatformStateObj(Response.GetPlatformState());
+	FJsonValueObject IOSSpecificObj(IOSSpecific);
+
+	TestTrue(TEXT("Parse"), FJsonValue::CompareEqual(PlatformStateObj, IOSSpecificObj));
+
+	return true;
+}
+
+IMPLEMENT_SIMPLE_AUTOMATION_TEST(TStartRecordingTakeResponseTest, "Plugin.CaptureProtocolStack.Control.ControlResponse.StartRecordingTake.Success", EAutomationTestFlags::ApplicationContextMask | EAutomationTestFlags::SmokeFilter)
+bool TStartRecordingTakeResponseTest::RunTest(const FString& InParameters)
+{
+    TSharedPtr<FJsonObject> Body;
+
+    FStartRecordingTakeResponse Response;
+
+    TestEqual(TEXT("GetAddresPath"), Response.GetAddressPath(), UE::CPS::AddressPaths::GStartRecordingTake);
+    TestTrue(TEXT("Parse"), Response.Parse(Body).HasValue());
+
+    return true;
+}
+
+IMPLEMENT_SIMPLE_AUTOMATION_TEST(TStopRecordingTakeResponseTest, "Plugin.CaptureProtocolStack.Control.ControlResponse.StopRecordingTake.Success", EAutomationTestFlags::ApplicationContextMask | EAutomationTestFlags::SmokeFilter)
+bool TStopRecordingTakeResponseTest::RunTest(const FString& InParameters)
+{
+	TSharedPtr<FJsonObject> Body = MakeShared<FJsonObject>();
+	Body->SetStringField(UE::CPS::Properties::GName, TEXT("TakeName"));
+
+    FStopRecordingTakeResponse Response;
+
+    TestEqual(TEXT("GetAddresPath"), Response.GetAddressPath(), UE::CPS::AddressPaths::GStopRecordingTake);
+	TestTrue(TEXT("Parse"), Response.Parse(Body).HasValue());
+
+	TestEqual(TEXT("Parse"), Response.GetTakeName(), TEXT("TakeName"));
+
+    return true;
+}
+
+IMPLEMENT_SIMPLE_AUTOMATION_TEST(TAbortRecordingTakeResponseTest, "Plugin.CaptureProtocolStack.Control.ControlResponse.AbortRecordingTake.Success", EAutomationTestFlags::ApplicationContextMask | EAutomationTestFlags::SmokeFilter)
+bool TAbortRecordingTakeResponseTest::RunTest(const FString& InParameters)
+{
+	TSharedPtr<FJsonObject> Body;
+
+	FAbortRecordingTakeResponse Response;
+
+	TestEqual(TEXT("GetAddresPath"), Response.GetAddressPath(), UE::CPS::AddressPaths::GAbortRecordingTake);
+	TestTrue(TEXT("Parse"), Response.Parse(Body).HasValue());
+
+	return true;
+}
+
+IMPLEMENT_SIMPLE_AUTOMATION_TEST(TGetTakeListResponseTest, "Plugin.CaptureProtocolStack.Control.ControlResponse.GetTakeList.Success", EAutomationTestFlags::ApplicationContextMask | EAutomationTestFlags::SmokeFilter)
+bool TGetTakeListResponseTest::RunTest(const FString& InParameters)
+{
+    TSharedPtr<FJsonObject> Body = MakeShared<FJsonObject>();
+
+    TArray<TSharedPtr<FJsonValue>> NamesJson = {
+        MakeShared<FJsonValueString>(TEXT("Name1")),
+        MakeShared<FJsonValueString>(TEXT("Name2")),
+        MakeShared<FJsonValueString>(TEXT("Name3")),
+    };
+
+    Body->SetArrayField(UE::CPS::Properties::GNames, NamesJson);
+
+    FGetTakeListResponse Response;
+
+    TestEqual(TEXT("GetAddresPath"), Response.GetAddressPath(), UE::CPS::AddressPaths::GGetTakeList);
+    TestTrue(TEXT("Parse"), Response.Parse(Body).HasValue());
+
+    const TArray<FString>& Names = Response.GetNames();
+    TestEqual(TEXT("Parse"), Names[0], TEXT("Name1"));
+    TestEqual(TEXT("Parse"), Names[1], TEXT("Name2"));
+    TestEqual(TEXT("Parse"), Names[2], TEXT("Name3"));
+    
+    return true;
+}
+
+IMPLEMENT_SIMPLE_AUTOMATION_TEST(TGetTakeMetadataResponseTest, "Plugin.CaptureProtocolStack.Control.ControlResponse.GetTakeMetadata.Success", EAutomationTestFlags::ApplicationContextMask | EAutomationTestFlags::SmokeFilter)
+bool TGetTakeMetadataResponseTest::RunTest(const FString& InParameters)
+{
+
+    FString Date = FDateTime::Now().ToIso8601();
+    TSharedPtr<FJsonObject> TakeJson = MakeShared<FJsonObject>();
+    TakeJson->SetStringField(UE::CPS::Properties::GName, TEXT("Name"));
+    TakeJson->SetStringField(UE::CPS::Properties::GSlateName, TEXT("Slate"));
+    TakeJson->SetNumberField(UE::CPS::Properties::GTakeNumber, 0);
+    TakeJson->SetStringField(UE::CPS::Properties::GDateTime, Date);
+    TakeJson->SetStringField(UE::CPS::Properties::GAppVersion, TEXT("AppVersion"));
+    TakeJson->SetStringField(UE::CPS::Properties::GModel, TEXT("Model"));
+    TakeJson->SetStringField(UE::CPS::Properties::GSubject, TEXT("Subject"));
+    TakeJson->SetStringField(UE::CPS::Properties::GScenario, TEXT("Scenario"));
+
+    TArray<TSharedPtr<FJsonValue>> TagsJson = {
+        MakeShared<FJsonValueString>(TEXT("Tag1")),
+        MakeShared<FJsonValueString>(TEXT("Tag2")),
+        MakeShared<FJsonValueString>(TEXT("Tag3")),
+    };
+
+    TakeJson->SetArrayField(UE::CPS::Properties::GTags, TagsJson);
+
+    TSharedPtr<FJsonObject> FileJson1 = MakeShared<FJsonObject>();
+    FileJson1->SetStringField(UE::CPS::Properties::GName, TEXT("File1"));
+    FileJson1->SetNumberField(UE::CPS::Properties::GLength, 1024);
+
+    TSharedPtr<FJsonObject> FileJson2 = MakeShared<FJsonObject>();
+    FileJson2->SetStringField(UE::CPS::Properties::GName, TEXT("File2"));
+    FileJson2->SetNumberField(UE::CPS::Properties::GLength, 1024);
+
+    TSharedPtr<FJsonObject> FileJson3 = MakeShared<FJsonObject>();
+    FileJson3->SetStringField(UE::CPS::Properties::GName, TEXT("File3"));
+    FileJson3->SetNumberField(UE::CPS::Properties::GLength, 1024);
+
+    TArray<TSharedPtr<FJsonValue>> FilesJson = {
+        MakeShared<FJsonValueObject>(FileJson1),
+        MakeShared<FJsonValueObject>(FileJson2),
+        MakeShared<FJsonValueObject>(FileJson3)
+    };
+
+    TakeJson->SetArrayField(UE::CPS::Properties::GFiles, FilesJson);
+
+    TSharedPtr<FJsonObject> VideoMetadataJson = MakeShared<FJsonObject>();
+	VideoMetadataJson->SetNumberField(UE::CPS::Properties::GFrames, 600);
+    VideoMetadataJson->SetNumberField(UE::CPS::Properties::GFrameRate, 60);
+    VideoMetadataJson->SetNumberField(UE::CPS::Properties::GHeight, 1024);
+    VideoMetadataJson->SetNumberField(UE::CPS::Properties::GWidth, 1024);
+
+    TakeJson->SetObjectField(UE::CPS::Properties::GVideo, VideoMetadataJson);
+
+    TSharedPtr<FJsonObject> AudioMetadataJson = MakeShared<FJsonObject>();
+    AudioMetadataJson->SetNumberField(UE::CPS::Properties::GChannels, 2);
+    AudioMetadataJson->SetNumberField(UE::CPS::Properties::GSampleRate, 44100);
+    AudioMetadataJson->SetNumberField(UE::CPS::Properties::GBitsPerChannel, 8);
+
+    TakeJson->SetObjectField(UE::CPS::Properties::GAudio, AudioMetadataJson);
+
+    TArray<TSharedPtr<FJsonValue>> TakesJson = {
+        MakeShared<FJsonValueObject>(TakeJson)
+    };
+
+    TSharedPtr<FJsonObject> Body = MakeShared<FJsonObject>();
+    Body->SetArrayField(UE::CPS::Properties::GTakes, TakesJson);
+
+    FGetTakeMetadataResponse Response;
+
+    TestEqual(TEXT("GetAddresPath"), Response.GetAddressPath(), UE::CPS::AddressPaths::GGetTakeMetadata);
+    TestTrue(TEXT("Parse"), Response.Parse(Body).HasValue());
+
+    const TArray<FGetTakeMetadataResponse::FTakeObject>& Takes = Response.GetTakes();
+    const FGetTakeMetadataResponse::FTakeObject& Take = Takes[0];
+
+    TestEqual(TEXT("Parse"), Take.Name, TEXT("Name"));
+    TestEqual(TEXT("Parse"), Take.Slate, TEXT("Slate"));
+    TestEqual(TEXT("Parse"), Take.TakeNumber, (uint16) 0);
+    TestEqual(TEXT("Parse"), Take.DateTime, Date);
+    TestEqual(TEXT("Parse"), Take.AppVersion, TEXT("AppVersion"));
+    TestEqual(TEXT("Parse"), Take.Model, TEXT("Model"));
+    TestEqual(TEXT("Parse"), Take.Subject, TEXT("Subject"));
+    TestEqual(TEXT("Parse"), Take.Scenario, TEXT("Scenario"));
+    
+    TestEqual(TEXT("Parse"), Take.Tags[0], TEXT("Tag1"));
+    TestEqual(TEXT("Parse"), Take.Tags[1], TEXT("Tag2"));
+    TestEqual(TEXT("Parse"), Take.Tags[2], TEXT("Tag3"));
+
+    TestEqual(TEXT("Parse"), Take.Files[0].Name, TEXT("File1"));
+    TestEqual(TEXT("Parse"), Take.Files[0].Length, (uint64) 1024);
+    TestEqual(TEXT("Parse"), Take.Files[1].Name, TEXT("File2"));
+    TestEqual(TEXT("Parse"), Take.Files[1].Length, (uint64) 1024);
+    TestEqual(TEXT("Parse"), Take.Files[2].Name, TEXT("File3"));
+    TestEqual(TEXT("Parse"), Take.Files[2].Length, (uint64) 1024);
+    
+	TestEqual(TEXT("Parse"), Take.Video.Frames, (uint64) 600);
+    TestEqual(TEXT("Parse"), Take.Video.FrameRate, (uint16) 60);
+    TestEqual(TEXT("Parse"), Take.Video.Height, (uint32) 1024);
+    TestEqual(TEXT("Parse"), Take.Video.Width, (uint32) 1024);
+
+    TestEqual(TEXT("Parse"), Take.Audio.Channels, (uint8) 2);
+    TestEqual(TEXT("Parse"), Take.Audio.SampleRate, (uint32) 44100);
+    TestEqual(TEXT("Parse"), Take.Audio.BitsPerChannel, (uint8) 8);
+
+    return true;
+}
+
+#endif // WITH_DEV_AUTOMATION_TESTS
\ No newline at end of file
diff --git a/Engine/Plugins/VirtualProduction/CaptureManager/Source/CaptureProtocolStack/Private/Tests/Control/TestControlUpdate.cpp b/Engine/Plugins/VirtualProduction/CaptureManager/Source/CaptureProtocolStack/Private/Tests/Control/TestControlUpdate.cpp
new file mode 100644
--- /dev/null
+++ b/Engine/Plugins/VirtualProduction/CaptureManager/Source/CaptureProtocolStack/Private/Tests/Control/TestControlUpdate.cpp
@@ -0,0 +1,147 @@
+// Copyright Epic Games, Inc. All Rights Reserved.
+
+#include "Control/Messages/ControlUpdate.h"
+
+#include "Control/Messages/ControlJsonUtilities.h"
+#include "Control/Messages/Constants.h"
+
+#include "Misc/AutomationTest.h"
+
+#if WITH_DEV_AUTOMATION_TESTS
+
+IMPLEMENT_SIMPLE_AUTOMATION_TEST(TSessionStoppedUpdateTest, "Plugin.CaptureProtocolStack.Control.ControlUpdate.SessionStopped.Success", EAutomationTestFlags::ApplicationContextMask | EAutomationTestFlags::SmokeFilter)
+bool TSessionStoppedUpdateTest::RunTest(const FString& InParameters)
+{
+    TSharedPtr<FJsonObject> Body;
+
+    FSessionStopped Update;
+
+    TestEqual(TEXT("GetAddresPath"), Update.GetAddressPath(), UE::CPS::AddressPaths::GSessionStopped);
+    TestTrue(TEXT("Parse"), Update.Parse(Body).HasValue());
+
+    return true;
+}
+
+IMPLEMENT_SIMPLE_AUTOMATION_TEST(TTakeAddedUpdateTest, "Plugin.CaptureProtocolStack.Control.ControlUpdate.TakeAdded.Success", EAutomationTestFlags::ApplicationContextMask | EAutomationTestFlags::SmokeFilter)
+bool TTakeAddedUpdateTest::RunTest(const FString& InParameters)
+{
+	TSharedPtr<FJsonObject> Body = MakeShared<FJsonObject>();
+	Body->SetStringField(UE::CPS::Properties::GName, TEXT("TakeName"));
+
+	FTakeAddedUpdate Update;
+
+	TestEqual(TEXT("GetAddresPath"), Update.GetAddressPath(), UE::CPS::AddressPaths::GTakeAdded);
+
+	TProtocolResult<void> UpdateParseResult = Update.Parse(Body);
+	TestTrue(TEXT("Parse"), UpdateParseResult.HasValue());
+	TestEqual(TEXT("Parse"), Update.GetTakeName(), TEXT("TakeName"));
+
+	return true;
+}
+
+IMPLEMENT_SIMPLE_AUTOMATION_TEST(TTakeRemovedUpdateTest, "Plugin.CaptureProtocolStack.Control.ControlUpdate.TakeRemoved.Success", EAutomationTestFlags::ApplicationContextMask | EAutomationTestFlags::SmokeFilter)
+bool TTakeRemovedUpdateTest::RunTest(const FString& InParameters)
+{
+	TSharedPtr<FJsonObject> Body = MakeShared<FJsonObject>();
+	Body->SetStringField(UE::CPS::Properties::GName, TEXT("TakeName"));
+
+	FTakeRemovedUpdate Update;
+
+	TestEqual(TEXT("GetAddresPath"), Update.GetAddressPath(), UE::CPS::AddressPaths::GTakeRemoved);
+
+	TProtocolResult<void> UpdateParseResult = Update.Parse(Body);
+	TestTrue(TEXT("Parse"), UpdateParseResult.HasValue());
+	TestEqual(TEXT("Parse"), Update.GetTakeName(), TEXT("TakeName"));
+
+	return true;
+}
+
+IMPLEMENT_SIMPLE_AUTOMATION_TEST(TTakeUpdatedUpdateTest, "Plugin.CaptureProtocolStack.Control.ControlUpdate.TakeUpdated.Success", EAutomationTestFlags::ApplicationContextMask | EAutomationTestFlags::SmokeFilter)
+bool TTakeUpdatedUpdateTest::RunTest(const FString& InParameters)
+{
+	TSharedPtr<FJsonObject> Body = MakeShared<FJsonObject>();
+	Body->SetStringField(UE::CPS::Properties::GName, TEXT("TakeName"));
+
+	FTakeUpdatedUpdate Update;
+
+	TestEqual(TEXT("GetAddresPath"), Update.GetAddressPath(), UE::CPS::AddressPaths::GTakeUpdated);
+
+	TProtocolResult<void> UpdateParseResult = Update.Parse(Body);
+	TestTrue(TEXT("Parse"), UpdateParseResult.HasValue());
+	TestEqual(TEXT("Parse"), Update.GetTakeName(), TEXT("TakeName"));
+
+	return true;
+}
+
+IMPLEMENT_SIMPLE_AUTOMATION_TEST(TRecordingStatusUpdateTest, "Plugin.CaptureProtocolStack.Control.ControlUpdate.RecordingStatus.Success", EAutomationTestFlags::ApplicationContextMask | EAutomationTestFlags::SmokeFilter)
+bool TRecordingStatusUpdateTest::RunTest(const FString& InParameters)
+{
+    TSharedPtr<FJsonObject> Body = MakeShared<FJsonObject>();
+    Body->SetBoolField(UE::CPS::Properties::GIsRecording, true);
+
+    FRecordingStatusUpdate Update;
+
+    TestEqual(TEXT("GetAddresPath"), Update.GetAddressPath(), UE::CPS::AddressPaths::GRecordingStatus);
+
+    TProtocolResult<void> UpdateParseResult = Update.Parse(Body);
+    TestTrue(TEXT("Parse"), UpdateParseResult.HasValue());
+    TestEqual(TEXT("Parse"), Update.IsRecording(), true);
+
+    return true;
+}
+
+IMPLEMENT_SIMPLE_AUTOMATION_TEST(TDiskCapacityUpdateTest, "Plugin.CaptureProtocolStack.Control.ControlUpdate.DiskCapacity.Success", EAutomationTestFlags::ApplicationContextMask | EAutomationTestFlags::SmokeFilter)
+bool TDiskCapacityUpdateTest::RunTest(const FString& InParameters)
+{
+    TSharedPtr<FJsonObject> Body = MakeShared<FJsonObject>();
+    Body->SetNumberField(UE::CPS::Properties::GTotal, 1024);
+    Body->SetNumberField(UE::CPS::Properties::GRemaining, 512);
+
+    FDiskCapacityUpdate Update;
+
+    TestEqual(TEXT("GetAddresPath"), Update.GetAddressPath(), UE::CPS::AddressPaths::GDiskCapacity);
+
+    TProtocolResult<void> UpdateParseResult = Update.Parse(Body);
+    TestTrue(TEXT("Parse"), UpdateParseResult.HasValue());
+    TestEqual(TEXT("Parse"), Update.GetTotal(), (uint64)1024);
+    TestEqual(TEXT("Parse"), Update.GetRemaining(), (uint64)512);
+
+    return true;
+}
+
+IMPLEMENT_SIMPLE_AUTOMATION_TEST(TBatteryPercentageUpdateTest, "Plugin.CaptureProtocolStack.Control.ControlUpdate.BatteryPercentage.Success", EAutomationTestFlags::ApplicationContextMask | EAutomationTestFlags::SmokeFilter)
+bool TBatteryPercentageUpdateTest::RunTest(const FString& InParameters)
+{
+    TSharedPtr<FJsonObject> Body = MakeShared<FJsonObject>();
+    Body->SetNumberField(UE::CPS::Properties::GLevel, 100.0);
+
+    FBatteryPercentageUpdate Update;
+
+    TestEqual(TEXT("GetAddresPath"), Update.GetAddressPath(), UE::CPS::AddressPaths::GBattery);
+
+    TProtocolResult<void> UpdateParseResult = Update.Parse(Body);
+    TestTrue(TEXT("Parse"), UpdateParseResult.HasValue());
+    TestEqual(TEXT("Parse"), Update.GetLevel(), (float) 100.0);
+
+    return true;
+}
+
+IMPLEMENT_SIMPLE_AUTOMATION_TEST(TThermalStateUpdateTest, "Plugin.CaptureProtocolStack.Control.ControlUpdate.ThermalState.Success", EAutomationTestFlags::ApplicationContextMask | EAutomationTestFlags::SmokeFilter)
+bool TThermalStateUpdateTest::RunTest(const FString& InParameters)
+{
+    TSharedPtr<FJsonObject> Body = MakeShared<FJsonObject>();
+
+    Body->SetStringField(UE::CPS::Properties::GState, TEXT("nominal"));
+
+    FThermalStateUpdate Update;
+
+    TestEqual(TEXT("GetAddresPath"), Update.GetAddressPath(), UE::CPS::AddressPaths::GThermalState);
+
+    TProtocolResult<void> UpdateParseResult = Update.Parse(Body);
+    TestTrue(TEXT("Parse"), UpdateParseResult.HasValue());
+    TestEqual(TEXT("Parse"), Update.GetState(), FThermalStateUpdate::EState::Nominal);
+
+    return true;
+}
+
+#endif // WITH_DEV_AUTOMATION_TESTS
\ No newline at end of file
diff --git a/Engine/Plugins/VirtualProduction/CaptureManager/Source/CaptureProtocolStack/Private/Tests/Discovery/TestDiscoveryNotify.cpp b/Engine/Plugins/VirtualProduction/CaptureManager/Source/CaptureProtocolStack/Private/Tests/Discovery/TestDiscoveryNotify.cpp
new file mode 100644
--- /dev/null
+++ b/Engine/Plugins/VirtualProduction/CaptureManager/Source/CaptureProtocolStack/Private/Tests/Discovery/TestDiscoveryNotify.cpp
@@ -0,0 +1,160 @@
+// Copyright Epic Games, Inc. All Rights Reserved.
+
+#include "Discovery/Messages/DiscoveryNotify.h"
+
+#include "Misc/AutomationTest.h"
+
+#if WITH_DEV_AUTOMATION_TESTS
+
+IMPLEMENT_SIMPLE_AUTOMATION_TEST(TDiscoveryNotifyDeserializeTest, "Plugin.CaptureProtocolStack.Discovery.DiscoveryNotify.Deserialize.Success", EAutomationTestFlags::ApplicationContextMask | EAutomationTestFlags::SmokeFilter)
+bool TDiscoveryNotifyDeserializeTest::RunTest(const FString& InParameters)
+{
+	// Prepare test
+	const TArray<uint8> Header = { 'C', 'P', 'S', 'D', 'I', 'S', 'C', 'O', 'V', 'E', 'R', '\0' };
+	const FDiscoveryPacket::EMessageType MessageType = FDiscoveryPacket::EMessageType::Notify;
+	const TArray<uint8> ServerId = { 0x0d, 0x0e, 0x0a, 0x0d, 0x0b, 0x0e, 0x0e, 0x0f, 0x0d, 0x0e, 0x0a, 0x0d, 0x0b, 0x0e, 0x0e, 0x0f };
+	const uint16 ControlPort = 8000;
+	const FDiscoveryNotify::EConnectionState ConnectionState = FDiscoveryNotify::EConnectionState::Online;
+	const TArray<uint16> SupportedVersions = { 1, 2, 3 };
+
+	TArray<uint8> Packet;
+	Packet.Reserve(Header.Num() + sizeof(MessageType) + ServerId.Num() + sizeof(ControlPort) + SupportedVersions.Num() * sizeof(uint16));
+	Packet.Append(Header);
+	Packet.Add(static_cast<uint8>(MessageType));
+	Packet.Append(ServerId);
+	Packet.Append((uint8*) &ControlPort, sizeof(ControlPort));
+	Packet.Add(static_cast<uint8>(ConnectionState));
+	Packet.Append((uint8*)SupportedVersions.GetData(), SupportedVersions.Num() * sizeof(uint16));
+
+	// Deserialization
+	TProtocolResult<FDiscoveryPacket> PacketResult = FDiscoveryPacket::Deserialize(Packet);
+	TestTrue(TEXT("Packet Deserialize"), PacketResult.HasValue());
+
+	FDiscoveryPacket DiscoveryPacket = PacketResult.StealValue();
+	TProtocolResult<FDiscoveryNotify> NotifyResult = FDiscoveryNotify::Deserialize(DiscoveryPacket);
+	TestTrue(TEXT("Notify Deserialize"), NotifyResult.HasValue());
+
+	FDiscoveryNotify Notify = NotifyResult.StealValue();
+	int32 ServerIdCompareResult = FMemory::Memcmp(Notify.GetServerId().GetData(), ServerId.GetData(), ServerId.Num());
+	TestTrue(TEXT("Notify Deserialize"), (ServerIdCompareResult == 0));
+	TestEqual(TEXT("Notify Deserialize"), Notify.GetControlPort(), ControlPort);
+	TestEqual(TEXT("Notify Deserialize"), Notify.GetConnectionState(), ConnectionState);
+	TestEqual(TEXT("Notify Deserialize"), Notify.GetSupportedVersions(), SupportedVersions);
+
+	return true;
+}
+
+IMPLEMENT_SIMPLE_AUTOMATION_TEST(TDiscoveryNotifyDeserializeTestInvalidMessageType, "Plugin.CaptureProtocolStack.Discovery.DiscoveryNotify.Deserialize.InvalidMessageType", EAutomationTestFlags::ApplicationContextMask | EAutomationTestFlags::SmokeFilter)
+bool TDiscoveryNotifyDeserializeTestInvalidMessageType::RunTest(const FString& InParameters)
+{
+	// Prepare test
+	const TArray<uint8> Header = { 'C', 'P', 'S', 'D', 'I', 'S', 'C', 'O', 'V', 'E', 'R', '\0' };
+	const FDiscoveryPacket::EMessageType MessageType = FDiscoveryPacket::EMessageType::Request;
+
+	TArray<uint8> Packet;
+	Packet.Reserve(Header.Num() + sizeof(MessageType));
+	Packet.Append(Header);
+	Packet.Add(static_cast<uint8>(MessageType));
+
+	// Deserialization
+	TProtocolResult<FDiscoveryPacket> PacketResult = FDiscoveryPacket::Deserialize(Packet);
+	TestTrue(TEXT("Packet Deserialize"), PacketResult.HasValue());
+
+	FDiscoveryPacket DiscoveryPacket = PacketResult.StealValue();
+	TProtocolResult<FDiscoveryNotify> NotifyResult = FDiscoveryNotify::Deserialize(DiscoveryPacket);
+	TestTrue(TEXT("Notify Deserialize"), NotifyResult.HasError());
+
+	return true;
+}
+
+IMPLEMENT_SIMPLE_AUTOMATION_TEST(TDiscoveryNotifyDeserializeTestInvalidSize, "Plugin.CaptureProtocolStack.Discovery.DiscoveryNotify.Deserialize.InvalidSize", EAutomationTestFlags::ApplicationContextMask | EAutomationTestFlags::SmokeFilter)
+bool TDiscoveryNotifyDeserializeTestInvalidSize::RunTest(const FString& InParameters)
+{
+	// Prepare test
+	const TArray<uint8> Header = { 'C', 'P', 'S', 'D', 'I', 'S', 'C', 'O', 'V', 'E', 'R', '\0' };
+	const FDiscoveryPacket::EMessageType MessageType = FDiscoveryPacket::EMessageType::Notify;
+	const TArray<uint8> ServerId = { 0x0d, 0x0e, 0x0a, 0x0d, 0x0b, 0x0e, 0x0e, 0x0f };
+
+	TArray<uint8> Packet;
+	Packet.Reserve(Header.Num() + sizeof(MessageType));
+	Packet.Append(Header);
+	Packet.Add(static_cast<uint8>(MessageType));
+	Packet.Append(ServerId);
+
+	// Deserialization
+	TProtocolResult<FDiscoveryPacket> PacketResult = FDiscoveryPacket::Deserialize(Packet);
+	TestTrue(TEXT("Packet Deserialize"), PacketResult.HasValue());
+
+	FDiscoveryPacket DiscoveryPacket = PacketResult.StealValue();
+	TProtocolResult<FDiscoveryNotify> NotifyResult = FDiscoveryNotify::Deserialize(DiscoveryPacket);
+	TestTrue(TEXT("Notify Deserialize"), NotifyResult.HasError());
+
+	return true;
+}
+
+IMPLEMENT_SIMPLE_AUTOMATION_TEST(TDiscoveryNotifyDeserializeTestInvalidConnectionState, "Plugin.CaptureProtocolStack.Discovery.DiscoveryNotify.Deserialize.InvalidConnectionState", EAutomationTestFlags::ApplicationContextMask | EAutomationTestFlags::SmokeFilter)
+bool TDiscoveryNotifyDeserializeTestInvalidConnectionState::RunTest(const FString& InParameters)
+{
+	// Prepare test
+	const TArray<uint8> Header = { 'C', 'P', 'S', 'D', 'I', 'S', 'C', 'O', 'V', 'E', 'R', '\0' };
+	const FDiscoveryPacket::EMessageType MessageType = FDiscoveryPacket::EMessageType::Notify;
+	const TArray<uint8> ServerId = { 0x0d, 0x0e, 0x0a, 0x0d, 0x0b, 0x0e, 0x0e, 0x0f, 0x0d, 0x0e, 0x0a, 0x0d, 0x0b, 0x0e, 0x0e, 0x0f };
+	const uint16 ControlPort = 8000;
+	const FDiscoveryNotify::EConnectionState ConnectionState = FDiscoveryNotify::EConnectionState::Invalid;
+	const TArray<uint16> SupportedVersions = { 1, 2, 3 };
+
+	TArray<uint8> Packet;
+	Packet.Reserve(Header.Num() + sizeof(MessageType) + ServerId.Num() + sizeof(ControlPort) + SupportedVersions.Num() * sizeof(uint16));
+	Packet.Append(Header);
+	Packet.Add(static_cast<uint8>(MessageType));
+	Packet.Append(ServerId);
+	Packet.Append((uint8*) &ControlPort, sizeof(ControlPort));
+	Packet.Add(static_cast<uint8>(ConnectionState));
+	Packet.Append((uint8*)SupportedVersions.GetData(), SupportedVersions.Num() * sizeof(uint16));
+
+	// Deserialization
+	TProtocolResult<FDiscoveryPacket> PacketResult = FDiscoveryPacket::Deserialize(Packet);
+	TestTrue(TEXT("Packet Deserialize"), PacketResult.HasValue());
+
+	FDiscoveryPacket DiscoveryPacket = PacketResult.StealValue();
+	TProtocolResult<FDiscoveryNotify> NotifyResult = FDiscoveryNotify::Deserialize(DiscoveryPacket);
+	TestTrue(TEXT("Notify Deserialize"), NotifyResult.HasError());
+
+	return true;
+}
+
+IMPLEMENT_SIMPLE_AUTOMATION_TEST(TDiscoveryNotifySerializeTest, "Plugin.CaptureProtocolStack.Discovery.DiscoveryNotify.Serialize.Success", EAutomationTestFlags::ApplicationContextMask | EAutomationTestFlags::SmokeFilter)
+bool TDiscoveryNotifySerializeTest::RunTest(const FString& InParameters)
+{
+	// Prepare test
+	const FDiscoveryPacket::EMessageType MessageType = FDiscoveryPacket::EMessageType::Notify;
+
+	const TArray<uint8> ServerId = { 0x0d, 0x0e, 0x0a, 0x0d, 0x0b, 0x0e, 0x0e, 0x0f, 0x0d, 0x0e, 0x0a, 0x0d, 0x0b, 0x0e, 0x0e, 0x0f };
+	const uint16 ControlPort = 8000;
+	const FDiscoveryNotify::EConnectionState ConnectionState = FDiscoveryNotify::EConnectionState::Online;
+	const TArray<uint16> SupportedVersions = { 1 };
+
+	TArray<uint8> Payload;
+	Payload.Reserve(ServerId.Num() + sizeof(ControlPort) + sizeof(ConnectionState) + SupportedVersions.Num() * sizeof(uint16));
+	Payload.Append(ServerId);
+	Payload.Append((uint8*) &ControlPort, sizeof(ControlPort));
+	Payload.Add(static_cast<uint8>(ConnectionState));
+	Payload.Append((uint8*) SupportedVersions.GetData(), SupportedVersions.Num() * sizeof(uint16));
+
+	FDiscoveryNotify::FServerId ServerIdStatic;
+	FMemory::Memcpy(ServerIdStatic.GetData(), ServerId.GetData(), ServerId.Num());
+
+	FDiscoveryNotify Notify(MoveTemp(ServerIdStatic), ControlPort, ConnectionState, SupportedVersions);
+
+	// Deserialization
+	TProtocolResult<FDiscoveryPacket> PacketResult = FDiscoveryNotify::Serialize(Notify);
+	TestTrue(TEXT("Packet Serialize"), PacketResult.HasValue());
+
+	FDiscoveryPacket DiscoveryPacket = PacketResult.StealValue();
+	TestEqual(TEXT("Notify Serialize"), DiscoveryPacket.GetMessageType(), MessageType);
+	TestEqual(TEXT("Notify Serialize"), DiscoveryPacket.GetPayload(), Payload);
+
+	return true;
+}
+
+#endif // WITH_DEV_AUTOMATION_TESTS
\ No newline at end of file
diff --git a/Engine/Plugins/VirtualProduction/CaptureManager/Source/CaptureProtocolStack/Private/Tests/Discovery/TestDiscoveryPacket.cpp b/Engine/Plugins/VirtualProduction/CaptureManager/Source/CaptureProtocolStack/Private/Tests/Discovery/TestDiscoveryPacket.cpp
new file mode 100644
--- /dev/null
+++ b/Engine/Plugins/VirtualProduction/CaptureManager/Source/CaptureProtocolStack/Private/Tests/Discovery/TestDiscoveryPacket.cpp
@@ -0,0 +1,116 @@
+// Copyright Epic Games, Inc. All Rights Reserved.
+
+#include "Discovery/Communication/DiscoveryPacket.h"
+
+#include "Misc/AutomationTest.h"
+
+#if WITH_DEV_AUTOMATION_TESTS
+
+IMPLEMENT_SIMPLE_AUTOMATION_TEST(TDiscoveryPacketDeserializeTest, "Plugin.CaptureProtocolStack.Discovery.DiscoveryPacket.Deserialize.Success", EAutomationTestFlags::ApplicationContextMask | EAutomationTestFlags::SmokeFilter)
+bool TDiscoveryPacketDeserializeTest::RunTest(const FString& InParameters)
+{
+	// Prepare test
+	const TArray<uint8> Header = { 'C', 'P', 'S', 'D', 'I', 'S', 'C', 'O', 'V', 'E', 'R', '\0' };
+	const FDiscoveryPacket::EMessageType MessageType = FDiscoveryPacket::EMessageType::Request;
+	const TArray<uint8> Payload = { 0x0d, 0x0e, 0x0a, 0x0d, 0x0b, 0x0e, 0x0e, 0x0f, 0x00 };
+
+	TArray<uint8> Packet;
+	Packet.Reserve(Header.Num() + sizeof(MessageType) + Payload.Num());
+	Packet.Append(Header);
+	Packet.Add(static_cast<uint8>(MessageType));
+	Packet.Append(Payload);
+
+	// Deserialization
+	TProtocolResult<FDiscoveryPacket> PacketResult = FDiscoveryPacket::Deserialize(Packet);
+	TestTrue(TEXT("Packet Deserialize"), PacketResult.HasValue());
+
+	FDiscoveryPacket DiscoveryPacket = PacketResult.StealValue();
+	TestEqual(TEXT("Packet Deserialize"), DiscoveryPacket.GetMessageType(), MessageType);
+	TestEqual(TEXT("Packet Deserialize"), DiscoveryPacket.GetPayload().Num(), Payload.Num());
+	TestEqual(TEXT("Packet Deserialize"), DiscoveryPacket.GetPayload(), Payload);
+
+	return true;
+}
+
+IMPLEMENT_SIMPLE_AUTOMATION_TEST(TDiscoveryPacketDeserializeTestInvalidHeaderSize, "Plugin.CaptureProtocolStack.Discovery.DiscoveryPacket.Deserialize.InvalidHeaderSize", EAutomationTestFlags::ApplicationContextMask | EAutomationTestFlags::SmokeFilter)
+bool TDiscoveryPacketDeserializeTestInvalidHeaderSize::RunTest(const FString& InParameters)
+{
+	// Prepare test
+	const TArray<uint8> Header = { 'A', 'A', 'A', '\0' };
+	const FDiscoveryPacket::EMessageType MessageType = FDiscoveryPacket::EMessageType::Request;
+
+	TArray<uint8> Packet;
+	Packet.Reserve(Header.Num() + sizeof(MessageType));
+	Packet.Append(Header);
+	Packet.Add(static_cast<uint8>(MessageType));
+
+	// Deserialization
+	TProtocolResult<FDiscoveryPacket> PacketResult = FDiscoveryPacket::Deserialize(Packet);
+	TestTrue(TEXT("Packet Deserialize"), PacketResult.HasError());
+
+	return true;
+}
+
+IMPLEMENT_SIMPLE_AUTOMATION_TEST(TDiscoveryPacketDeserializeTestInvalidHeader, "Plugin.CaptureProtocolStack.Discovery.DiscoveryPacket.Deserialize.InvalidHeader", EAutomationTestFlags::ApplicationContextMask | EAutomationTestFlags::SmokeFilter)
+bool TDiscoveryPacketDeserializeTestInvalidHeader::RunTest(const FString& InParameters)
+{
+	// Prepare test
+	const TArray<uint8> Header = { 'A', 'A', 'A', 'A', 'A', 'A', 'A', 'A', 'A', 'A', 'A', '\0' };
+	const FDiscoveryPacket::EMessageType MessageType = FDiscoveryPacket::EMessageType::Request;
+
+	TArray<uint8> Packet;
+	Packet.Reserve(Header.Num() + sizeof(MessageType));
+	Packet.Append(Header);
+	Packet.Add(static_cast<uint8>(MessageType));
+
+	// Deserialization
+	TProtocolResult<FDiscoveryPacket> PacketResult = FDiscoveryPacket::Deserialize(Packet);
+	TestTrue(TEXT("Packet Deserialize"), PacketResult.HasError());
+
+	return true;
+}
+
+IMPLEMENT_SIMPLE_AUTOMATION_TEST(TDiscoveryPacketDeserializeTestInvalidMessageType, "Plugin.CaptureProtocolStack.Discovery.DiscoveryPacket.Deserialize.InvalidMessageType", EAutomationTestFlags::ApplicationContextMask | EAutomationTestFlags::SmokeFilter)
+bool TDiscoveryPacketDeserializeTestInvalidMessageType::RunTest(const FString& InParameters)
+{
+	// Prepare test
+	const TArray<uint8> Header = { 'C', 'P', 'S', 'D', 'I', 'S', 'C', 'O', 'V', 'E', 'R', '\0' };
+	const FDiscoveryPacket::EMessageType MessageType = FDiscoveryPacket::EMessageType::Invalid;
+
+	TArray<uint8> Packet;
+	Packet.Reserve(Header.Num() + sizeof(MessageType));
+	Packet.Append(Header);
+	Packet.Add(static_cast<uint8>(MessageType));
+
+	// Deserialization
+	TProtocolResult<FDiscoveryPacket> PacketResult = FDiscoveryPacket::Deserialize(Packet);
+	TestTrue(TEXT("Packet Deserialize"), PacketResult.HasError());
+
+	return true;
+}
+
+IMPLEMENT_SIMPLE_AUTOMATION_TEST(TDiscoveryPacketSerializeTest, "Plugin.CaptureProtocolStack.Discovery.DiscoveryPacket.Serialize.Success", EAutomationTestFlags::ApplicationContextMask | EAutomationTestFlags::SmokeFilter)
+bool TDiscoveryPacketSerializeTest::RunTest(const FString& InParameters)
+{
+	// Prepare test
+	const TArray<uint8> Header = { 'C', 'P', 'S', 'D', 'I', 'S', 'C', 'O', 'V', 'E', 'R', '\0' };
+	const FDiscoveryPacket::EMessageType MessageType = FDiscoveryPacket::EMessageType::Request;
+
+	TArray<uint8> Packet;
+	Packet.Reserve(Header.Num() + sizeof(MessageType));
+	Packet.Append(Header);
+	Packet.Add(static_cast<uint8>(MessageType));
+
+	FDiscoveryPacket DiscoveryPacket(MessageType, TArray<uint8>());
+
+	// Deserialization
+	TProtocolResult<TArray<uint8>> DataResult = FDiscoveryPacket::Serialize(DiscoveryPacket);
+	TestTrue(TEXT("Packet Serialize"), DataResult.HasValue());
+
+	TArray<uint8> Data = DataResult.StealValue();
+	TestEqual(TEXT("Packet Serialize"), Packet, Data);
+
+	return true;
+}
+
+#endif //WITH_DEV_AUTOMATION_TESTS
\ No newline at end of file
diff --git a/Engine/Plugins/VirtualProduction/CaptureManager/Source/CaptureProtocolStack/Private/Tests/Discovery/TestDiscoveryRequest.cpp b/Engine/Plugins/VirtualProduction/CaptureManager/Source/CaptureProtocolStack/Private/Tests/Discovery/TestDiscoveryRequest.cpp
new file mode 100644
--- /dev/null
+++ b/Engine/Plugins/VirtualProduction/CaptureManager/Source/CaptureProtocolStack/Private/Tests/Discovery/TestDiscoveryRequest.cpp
@@ -0,0 +1,75 @@
+// Copyright Epic Games, Inc. All Rights Reserved.
+
+#include "Discovery/Messages/DiscoveryRequest.h"
+
+#include "Misc/AutomationTest.h"
+
+#if WITH_DEV_AUTOMATION_TESTS
+
+IMPLEMENT_SIMPLE_AUTOMATION_TEST(TDiscoveryRequestDeserializeTest, "Plugin.CaptureProtocolStack.Discovery.DiscoveryRequest.Deserialize.Success", EAutomationTestFlags::ApplicationContextMask | EAutomationTestFlags::SmokeFilter)
+bool TDiscoveryRequestDeserializeTest::RunTest(const FString& InParameters)
+{
+	// Prepare test
+	const TArray<uint8> Header = { 'C', 'P', 'S', 'D', 'I', 'S', 'C', 'O', 'V', 'E', 'R', '\0' };
+	const FDiscoveryPacket::EMessageType MessageType = FDiscoveryPacket::EMessageType::Request;
+
+	TArray<uint8> Packet;
+	Packet.Reserve(Header.Num() + sizeof(MessageType));
+	Packet.Append(Header);
+	Packet.Add(static_cast<uint8>(MessageType));
+
+	// Deserialization
+	TProtocolResult<FDiscoveryPacket> PacketResult = FDiscoveryPacket::Deserialize(Packet);
+	TestTrue(TEXT("Packet Deserialize"), PacketResult.HasValue());
+
+	FDiscoveryPacket DiscoveryPacket = PacketResult.StealValue();
+	TProtocolResult<FDiscoveryRequest> RequestResult = FDiscoveryRequest::Deserialize(DiscoveryPacket);
+	TestTrue(TEXT("Request Deserialize"), RequestResult.HasValue());
+
+	return true;
+}
+
+IMPLEMENT_SIMPLE_AUTOMATION_TEST(TDiscoveryRequestDeserializeTestInvalidMessageType, "Plugin.CaptureProtocolStack.Discovery.DiscoveryRequest.Deserialize.InvalidMessageType", EAutomationTestFlags::ApplicationContextMask | EAutomationTestFlags::SmokeFilter)
+bool TDiscoveryRequestDeserializeTestInvalidMessageType::RunTest(const FString& InParameters)
+{
+	// Prepare test
+	const TArray<uint8> Header = { 'C', 'P', 'S', 'D', 'I', 'S', 'C', 'O', 'V', 'E', 'R', '\0' };
+	const FDiscoveryPacket::EMessageType MessageType = FDiscoveryPacket::EMessageType::Response;
+
+	TArray<uint8> Packet;
+	Packet.Reserve(Header.Num() + sizeof(MessageType));
+	Packet.Append(Header);
+	Packet.Add(static_cast<uint8>(MessageType));
+
+	// Deserialization
+	TProtocolResult<FDiscoveryPacket> PacketResult = FDiscoveryPacket::Deserialize(Packet);
+	TestTrue(TEXT("Packet Deserialize"), PacketResult.HasValue());
+
+	FDiscoveryPacket DiscoveryPacket = PacketResult.StealValue();
+
+	TProtocolResult<FDiscoveryRequest> RequestResult = FDiscoveryRequest::Deserialize(DiscoveryPacket);
+	TestTrue(TEXT("Request Deserialize"), RequestResult.HasError());
+
+	return true;
+}
+
+IMPLEMENT_SIMPLE_AUTOMATION_TEST(TDiscoveryRequestSerializeTest, "Plugin.CaptureProtocolStack.Discovery.DiscoveryRequest.Serialize.Success", EAutomationTestFlags::ApplicationContextMask | EAutomationTestFlags::SmokeFilter)
+bool TDiscoveryRequestSerializeTest::RunTest(const FString& InParameters)
+{
+	// Prepare test
+	const FDiscoveryPacket::EMessageType MessageType = FDiscoveryPacket::EMessageType::Request;
+
+	FDiscoveryRequest Request;
+
+	// Deserialization
+	TProtocolResult<FDiscoveryPacket> PacketResult = FDiscoveryRequest::Serialize(Request);
+	TestTrue(TEXT("Request Serialize"), PacketResult.HasValue());
+
+	FDiscoveryPacket DiscoveryPacket = PacketResult.StealValue();
+	TestEqual(TEXT("Request Serialize"), DiscoveryPacket.GetMessageType(), MessageType);
+	TestTrue(TEXT("Request Serialize"), DiscoveryPacket.GetPayload().IsEmpty());
+
+	return true;
+}
+
+#endif // WITH_DEV_AUTOMATION_TESTS
\ No newline at end of file
diff --git a/Engine/Plugins/VirtualProduction/CaptureManager/Source/CaptureProtocolStack/Private/Tests/Discovery/TestDiscoveryResponse.cpp b/Engine/Plugins/VirtualProduction/CaptureManager/Source/CaptureProtocolStack/Private/Tests/Discovery/TestDiscoveryResponse.cpp
new file mode 100644
--- /dev/null
+++ b/Engine/Plugins/VirtualProduction/CaptureManager/Source/CaptureProtocolStack/Private/Tests/Discovery/TestDiscoveryResponse.cpp
@@ -0,0 +1,124 @@
+// Copyright Epic Games, Inc. All Rights Reserved.
+
+#include "Discovery/Messages/DiscoveryResponse.h"
+
+#include "Misc/AutomationTest.h"
+
+#if WITH_DEV_AUTOMATION_TESTS
+
+IMPLEMENT_SIMPLE_AUTOMATION_TEST(TDiscoveryResponseDeserializeTest, "Plugin.CaptureProtocolStack.Discovery.DiscoveryResponse.Deserialize.Success", EAutomationTestFlags::ApplicationContextMask | EAutomationTestFlags::SmokeFilter)
+bool TDiscoveryResponseDeserializeTest::RunTest(const FString& InParameters)
+{
+	// Prepare test
+	const TArray<uint8> Header = { 'C', 'P', 'S', 'D', 'I', 'S', 'C', 'O', 'V', 'E', 'R', '\0' };
+	const FDiscoveryPacket::EMessageType MessageType = FDiscoveryPacket::EMessageType::Response;
+	const TArray<uint8> ServerId = { 0x0d, 0x0e, 0x0a, 0x0d, 0x0b, 0x0e, 0x0e, 0x0f, 0x0d, 0x0e, 0x0a, 0x0d, 0x0b, 0x0e, 0x0e, 0x0f };
+	const uint16 ControlPort = 8000;
+	const TArray<uint16> SupportedVersions = { 1, 2, 3 };
+
+	TArray<uint8> Packet;
+	Packet.Reserve(Header.Num() + sizeof(MessageType) + ServerId.Num() + sizeof(ControlPort) + SupportedVersions.Num() * sizeof(uint16));
+	Packet.Append(Header);
+	Packet.Add(static_cast<uint8>(MessageType));
+	Packet.Append(ServerId);
+	Packet.Append((uint8*) &ControlPort, sizeof(ControlPort));
+	Packet.Append((uint8*)SupportedVersions.GetData(), SupportedVersions.Num() * sizeof(uint16));
+
+	// Deserialization
+	TProtocolResult<FDiscoveryPacket> PacketResult = FDiscoveryPacket::Deserialize(Packet);
+	TestTrue(TEXT("Packet Deserialize"), PacketResult.HasValue());
+
+	FDiscoveryPacket DiscoveryPacket = PacketResult.StealValue();
+	TProtocolResult<FDiscoveryResponse> ResponseResult = FDiscoveryResponse::Deserialize(DiscoveryPacket);
+	TestTrue(TEXT("Response Deserialize"), ResponseResult.HasValue());
+
+	FDiscoveryResponse Response = ResponseResult.StealValue();
+	int32 ServerIdCompareResult = FMemory::Memcmp(Response.GetServerId().GetData(), ServerId.GetData(), ServerId.Num());
+	TestTrue(TEXT("Response Deserialize"), (ServerIdCompareResult == 0));
+	TestEqual(TEXT("Response Deserialize"), Response.GetControlPort(), ControlPort);
+	TestEqual(TEXT("Response Deserialize"), Response.GetSupportedVersions(), SupportedVersions);
+
+	return true;
+}
+
+IMPLEMENT_SIMPLE_AUTOMATION_TEST(TDiscoveryResponseDeserializeTestInvalidMessageType, "Plugin.CaptureProtocolStack.Discovery.DiscoveryResponse.Deserialize.InvalidMessageType", EAutomationTestFlags::ApplicationContextMask | EAutomationTestFlags::SmokeFilter)
+bool TDiscoveryResponseDeserializeTestInvalidMessageType::RunTest(const FString& InParameters)
+{
+	// Prepare test
+	const TArray<uint8> Header = { 'C', 'P', 'S', 'D', 'I', 'S', 'C', 'O', 'V', 'E', 'R', '\0' };
+	const FDiscoveryPacket::EMessageType MessageType = FDiscoveryPacket::EMessageType::Request;
+
+	TArray<uint8> Packet;
+	Packet.Reserve(Header.Num() + sizeof(MessageType));
+	Packet.Append(Header);
+	Packet.Add(static_cast<uint8>(MessageType));
+
+	// Deserialization
+	TProtocolResult<FDiscoveryPacket> PacketResult = FDiscoveryPacket::Deserialize(Packet);
+	TestTrue(TEXT("Packet Deserialize"), PacketResult.HasValue());
+
+	FDiscoveryPacket DiscoveryPacket = PacketResult.StealValue();
+	TProtocolResult<FDiscoveryResponse> ResponseResult = FDiscoveryResponse::Deserialize(DiscoveryPacket);
+	TestTrue(TEXT("Response Deserialize"), ResponseResult.HasError());
+
+	return true;
+}
+
+IMPLEMENT_SIMPLE_AUTOMATION_TEST(TDiscoveryResponseDeserializeTestInvalidSize, "Plugin.CaptureProtocolStack.Discovery.DiscoveryResponse.Deserialize.InvalidSize", EAutomationTestFlags::ApplicationContextMask | EAutomationTestFlags::SmokeFilter)
+bool TDiscoveryResponseDeserializeTestInvalidSize::RunTest(const FString& InParameters)
+{
+	// Prepare test
+	const TArray<uint8> Header = { 'C', 'P', 'S', 'D', 'I', 'S', 'C', 'O', 'V', 'E', 'R', '\0' };
+	const FDiscoveryPacket::EMessageType MessageType = FDiscoveryPacket::EMessageType::Notify;
+	const TArray<uint8> ServerId = { 0x0d, 0x0e, 0x0a, 0x0d, 0x0b, 0x0e, 0x0e, 0x0f };
+
+	TArray<uint8> Packet;
+	Packet.Reserve(Header.Num() + sizeof(MessageType));
+	Packet.Append(Header);
+	Packet.Add(static_cast<uint8>(MessageType));
+	Packet.Append(ServerId);
+
+	// Deserialization
+	TProtocolResult<FDiscoveryPacket> PacketResult = FDiscoveryPacket::Deserialize(Packet);
+	TestTrue(TEXT("Packet Deserialize"), PacketResult.HasValue());
+
+	FDiscoveryPacket DiscoveryPacket = PacketResult.StealValue();
+	TProtocolResult<FDiscoveryResponse> ResponseResult = FDiscoveryResponse::Deserialize(DiscoveryPacket);
+	TestTrue(TEXT("Response Deserialize"), ResponseResult.HasError());
+
+	return true;
+}
+
+IMPLEMENT_SIMPLE_AUTOMATION_TEST(TDiscoveryResponseSerializeTest, "Plugin.CaptureProtocolStack.Discovery.DiscoveryResponse.Serialize.Success", EAutomationTestFlags::ApplicationContextMask | EAutomationTestFlags::SmokeFilter)
+bool TDiscoveryResponseSerializeTest::RunTest(const FString& InParameters)
+{
+	// Prepare test
+	const FDiscoveryPacket::EMessageType MessageType = FDiscoveryPacket::EMessageType::Response;
+
+	const TArray<uint8> ServerId = { 0x0d, 0x0e, 0x0a, 0x0d, 0x0b, 0x0e, 0x0e, 0x0f, 0x0d, 0x0e, 0x0a, 0x0d, 0x0b, 0x0e, 0x0e, 0x0f };
+	const uint16 ControlPort = 8000;
+	const TArray<uint16> SupportedVersions = { 1 };
+
+	TArray<uint8> Payload;
+	Payload.Reserve(ServerId.Num() + sizeof(ControlPort) + SupportedVersions.Num() * sizeof(uint16));
+	Payload.Append(ServerId);
+	Payload.Append((uint8*) &ControlPort, sizeof(ControlPort));
+	Payload.Append((uint8*) SupportedVersions.GetData(), SupportedVersions.Num() * sizeof(uint16));
+
+	FDiscoveryResponse::FServerId ServerIdStatic;
+	FMemory::Memcpy(ServerIdStatic.GetData(), ServerId.GetData(), ServerId.Num());
+
+	FDiscoveryResponse Response(MoveTemp(ServerIdStatic), ControlPort, SupportedVersions);
+
+	// Deserialization
+	TProtocolResult<FDiscoveryPacket> PacketResult = FDiscoveryResponse::Serialize(Response);
+	TestTrue(TEXT("Packet Serialize"), PacketResult.HasValue());
+
+	FDiscoveryPacket DiscoveryPacket = PacketResult.StealValue();
+	TestEqual(TEXT("Response Serialize"), DiscoveryPacket.GetMessageType(), MessageType);
+	TestEqual(TEXT("Response Serialize"), DiscoveryPacket.GetPayload(), Payload);
+
+	return true;
+}
+
+#endif // WITH_DEV_AUTOMATION_TESTS
\ No newline at end of file
diff --git a/Engine/Plugins/VirtualProduction/CaptureManager/Source/CaptureProtocolStack/Private/Tests/Export/TestExportHeader.cpp b/Engine/Plugins/VirtualProduction/CaptureManager/Source/CaptureProtocolStack/Private/Tests/Export/TestExportHeader.cpp
new file mode 100644
--- /dev/null
+++ b/Engine/Plugins/VirtualProduction/CaptureManager/Source/CaptureProtocolStack/Private/Tests/Export/TestExportHeader.cpp
@@ -0,0 +1,127 @@
+// Copyright Epic Games, Inc. All Rights Reserved.
+
+#include "ExportClient/Communication/ExportHeader.h"
+
+#include "Misc/AutomationTest.h"
+
+#include "Tests/Utility.h"
+
+#if WITH_DEV_AUTOMATION_TESTS
+
+IMPLEMENT_SIMPLE_AUTOMATION_TEST(TExportHeaderDeserializeOneTest, "Plugin.CaptureProtocolStack.Export.ExportHeader.DeserializeOne.Success", EAutomationTestFlags::ApplicationContextMask | EAutomationTestFlags::SmokeFilter)
+bool TExportHeaderDeserializeOneTest::RunTest(const FString& InParameters)
+{
+	// Prepare test
+	const TArray<uint8> Header = { 'C', 'P', 'S', 'E', 'X', 'P', 'O', 'R', 'T', '\0' };
+	const uint16 Version = 1;
+	const uint32 TransactionId = 11223344;
+
+	TArray<uint8> Data;
+	Data.Append(Header);
+	Data.Append(reinterpret_cast<const uint8*>(&Version), sizeof(Version));
+	Data.Append(reinterpret_cast<const uint8*>(&TransactionId), sizeof(TransactionId));
+
+	FDataProvider DataProvider(MoveTemp(Data));
+
+	// Deserialization
+	TProtocolResult<FExportHeader> ExportHeaderResult = FExportHeader::Deserialize(DataProvider);
+	TestTrue(TEXT("Header Deserialize"), ExportHeaderResult.HasValue());
+
+	FExportHeader ExportHeader = ExportHeaderResult.StealValue();
+
+	TestEqual(TEXT("Header Deserialize"), ExportHeader.GetVersion(), Version);
+	TestEqual(TEXT("Header Deserialize"), ExportHeader.GetTransactionId(), TransactionId);
+
+	return true;
+}
+
+IMPLEMENT_SIMPLE_AUTOMATION_TEST(TExportHeaderDeserializeOneTestInvalidHeader, "Plugin.CaptureProtocolStack.Export.ExportHeader.DeserializeOne.InvalidHeader", EAutomationTestFlags::ApplicationContextMask | EAutomationTestFlags::SmokeFilter)
+bool TExportHeaderDeserializeOneTestInvalidHeader::RunTest(const FString& InParameters)
+{
+	// Prepare test
+	const TArray<uint8> Header = { 'A', 'A', 'A', 'A', 'A', 'A', 'A', 'A', 'A', '\0' };
+	const uint16 Version = 1;
+	const uint32 TransactionId = 11223344;
+
+	TArray<uint8> Data;
+	Data.Append(Header);
+	Data.Append(reinterpret_cast<const uint8*>(&Version), sizeof(Version));
+	Data.Append(reinterpret_cast<const uint8*>(&TransactionId), sizeof(TransactionId));
+
+	FDataProvider DataProvider(MoveTemp(Data));
+
+	// Deserialization
+	TProtocolResult<FExportHeader> ExportHeaderResult = FExportHeader::Deserialize(DataProvider);
+	TestTrue(TEXT("Header Deserialize"), ExportHeaderResult.HasError());
+
+	return true;
+}
+
+IMPLEMENT_SIMPLE_AUTOMATION_TEST(TExportHeaderDeserializeOneTestInvalidSize, "Plugin.CaptureProtocolStack.Export.ExportHeader.DeserializeOne.InvalidSize", EAutomationTestFlags::ApplicationContextMask | EAutomationTestFlags::SmokeFilter)
+bool TExportHeaderDeserializeOneTestInvalidSize::RunTest(const FString& InParameters)
+{
+	// Prepare test
+	const TArray<uint8> Header = { 'C', 'P', 'S', 'E', 'X', 'P', 'O', 'R', 'T', '\0' };
+
+	TArray<uint8> Data;
+	Data.Reserve(Header.Num());
+	Data.Append(Header);
+
+	FDataProvider DataProvider(MoveTemp(Data));
+
+	// Deserialization
+	TProtocolResult<FExportHeader> ExportHeaderResult = FExportHeader::Deserialize(DataProvider);
+	TestTrue(TEXT("Header Deserialize"), ExportHeaderResult.HasError());
+
+	return true;
+}
+
+IMPLEMENT_SIMPLE_AUTOMATION_TEST(TExportHeaderSerializeOneTest, "Plugin.CaptureProtocolStack.Export.ExportHeader.SerializeOne.Success", EAutomationTestFlags::ApplicationContextMask | EAutomationTestFlags::SmokeFilter)
+bool TExportHeaderSerializeOneTest::RunTest(const FString& InParameters)
+{
+	// Prepare test
+	const TArray<uint8> Header = { 'C', 'P', 'S', 'E', 'X', 'P', 'O', 'R', 'T', '\0' };
+	const uint16 Version = 1;
+	const uint32 TransactionId = 11223344;
+
+	TArray<uint8> Data;
+	Data.Append(Header);
+	Data.Append(reinterpret_cast<const uint8*>(&Version), sizeof(Version));
+	Data.Append(reinterpret_cast<const uint8*>(&TransactionId), sizeof(TransactionId));
+
+	FDataSender DataSender;
+
+	// Serialization
+	FExportHeader ExportHeader(Version, TransactionId);
+	TProtocolResult<void> ExportHeaderResult = FExportHeader::Serialize(ExportHeader, DataSender);
+	TestTrue(TEXT("Header Serialize"), ExportHeaderResult.HasValue());
+
+	TestEqual(TEXT("Header Serialize"), DataSender.GetData(), Data);
+
+	return true;
+}
+
+IMPLEMENT_SIMPLE_AUTOMATION_TEST(TExportHeaderSerializeOneTestError, "Plugin.CaptureProtocolStack.Export.ExportHeader.SerializeOne.Error", EAutomationTestFlags::ApplicationContextMask | EAutomationTestFlags::SmokeFilter)
+bool TExportHeaderSerializeOneTestError::RunTest(const FString& InParameters)
+{
+	// Prepare test
+	const TArray<uint8> Header = { 'C', 'P', 'S', 'E', 'X', 'P', 'O', 'R', 'T', '\0' };
+	const uint16 Version = 1;
+	const uint32 TransactionId = 11223344;
+
+	TArray<uint8> Data;
+	Data.Append(Header);
+	Data.Append(reinterpret_cast<const uint8*>(&Version), sizeof(Version));
+	Data.Append(reinterpret_cast<const uint8*>(&TransactionId), sizeof(TransactionId));
+
+	FFailedDataSender DataSender;
+
+	// Serialization
+	FExportHeader ExportHeader(Version, TransactionId);
+	TProtocolResult<void> ExportHeaderResult = FExportHeader::Serialize(ExportHeader, DataSender);
+	TestTrue(TEXT("Header Serialize"), ExportHeaderResult.HasError());
+
+	return true;
+}
+
+#endif // WITH_DEV_AUTOMATION_TESTS
\ No newline at end of file
diff --git a/Engine/Plugins/VirtualProduction/CaptureManager/Source/CaptureProtocolStack/Private/Tests/Export/TestExportRequest.cpp b/Engine/Plugins/VirtualProduction/CaptureManager/Source/CaptureProtocolStack/Private/Tests/Export/TestExportRequest.cpp
new file mode 100644
--- /dev/null
+++ b/Engine/Plugins/VirtualProduction/CaptureManager/Source/CaptureProtocolStack/Private/Tests/Export/TestExportRequest.cpp
@@ -0,0 +1,119 @@
+// Copyright Epic Games, Inc. All Rights Reserved.
+
+#include "ExportClient/Messages/ExportRequest.h"
+
+#include "Misc/AutomationTest.h"
+
+#include "Tests/Utility.h"
+
+#if WITH_DEV_AUTOMATION_TESTS
+
+IMPLEMENT_SIMPLE_AUTOMATION_TEST(TExportRequestDeserializeOneTest, "Plugin.CaptureProtocolStack.Export.ExportRequest.DeserializeOne.Success", EAutomationTestFlags::ApplicationContextMask | EAutomationTestFlags::SmokeFilter)
+bool TExportRequestDeserializeOneTest::RunTest(const FString& InParameters)
+{
+	// Prepare test
+	const FString TakeName = TEXT("TakeName");
+	const FString FileName = TEXT("FileName");
+	const uint16 TakeNameLen = TakeName.Len();
+	const uint16 FileNameLen = FileName.Len();
+
+	const uint64 Offset = 0;
+
+	TArray<uint8> Packet;
+	Packet.Append(reinterpret_cast<const uint8*>(&TakeNameLen), sizeof(TakeNameLen));
+	Packet.Append(reinterpret_cast<const uint8*>(TCHAR_TO_UTF8(*TakeName)), TakeName.Len());
+	Packet.Append(reinterpret_cast<const uint8*>(&FileNameLen), sizeof(FileNameLen));
+	Packet.Append(reinterpret_cast<const uint8*>(TCHAR_TO_UTF8(*FileName)), FileName.Len());
+	Packet.Append(reinterpret_cast<const uint8*>(&Offset), sizeof(Offset));
+
+	FDataProvider DataProvider(MoveTemp(Packet));
+
+	// Deserialization
+	TProtocolResult<FExportRequest> ExportRequestResult = FExportRequest::Deserialize(DataProvider);
+	TestTrue(TEXT("Request Deserialize"), ExportRequestResult.HasValue());
+
+	FExportRequest ExportRequest = ExportRequestResult.StealValue();
+
+	TestEqual(TEXT("Request Deserialize"), ExportRequest.GetTakeName(), TakeName);
+	TestEqual(TEXT("Request Deserialize"), ExportRequest.GetFileName(), FileName);
+	TestEqual(TEXT("Request Deserialize"), ExportRequest.GetOffset(), Offset);
+
+	return true;
+}
+
+IMPLEMENT_SIMPLE_AUTOMATION_TEST(TExportRequestDeserializeOneTestInvalidSize, "Plugin.CaptureProtocolStack.Export.ExportRequest.DeserializeOne.InvalidSize", EAutomationTestFlags::ApplicationContextMask | EAutomationTestFlags::SmokeFilter)
+bool TExportRequestDeserializeOneTestInvalidSize::RunTest(const FString& InParameters)
+{
+	// Prepare test
+	TArray<uint8> Packet;
+
+	FDataProvider DataProvider(MoveTemp(Packet));
+
+	// Deserialization
+	TProtocolResult<FExportRequest> ExportRequestResult = FExportRequest::Deserialize(DataProvider);
+	TestTrue(TEXT("Request Deserialize"), ExportRequestResult.HasError());
+
+	return true;
+}
+
+IMPLEMENT_SIMPLE_AUTOMATION_TEST(TExportRequestSerializeOneTest, "Plugin.CaptureProtocolStack.Export.ExportRequest.SerializeOne.Success", EAutomationTestFlags::ApplicationContextMask | EAutomationTestFlags::SmokeFilter)
+bool TExportRequestSerializeOneTest::RunTest(const FString& InParameters)
+{
+	// Prepare test
+	const FString TakeName = TEXT("TakeName");
+	const FString FileName = TEXT("FileName");
+	const uint16 TakeNameLen = TakeName.Len();
+	const uint16 FileNameLen = FileName.Len();
+
+	const uint64 Offset = 0;
+
+	TArray<uint8> Packet;
+	Packet.Append(reinterpret_cast<const uint8*>(&TakeNameLen), sizeof(TakeNameLen));
+	Packet.Append(reinterpret_cast<const uint8*>(TCHAR_TO_UTF8(*TakeName)), TakeName.Len());
+	Packet.Append(reinterpret_cast<const uint8*>(&FileNameLen), sizeof(FileNameLen));
+	Packet.Append(reinterpret_cast<const uint8*>(TCHAR_TO_UTF8(*FileName)), FileName.Len());
+	Packet.Append(reinterpret_cast<const uint8*>(&Offset), sizeof(Offset));
+
+	FDataSender DataSender;
+
+	FExportRequest Request(TakeName, FileName, Offset);
+
+	// Serialization
+	TProtocolResult<void> ExportRequestResult = FExportRequest::Serialize(Request, DataSender);
+	TestTrue(TEXT("Request Serialize"), ExportRequestResult.HasValue());
+
+	TestEqual(TEXT("Request Serialize"), DataSender.GetData(), Packet);
+
+	return true;
+}
+
+IMPLEMENT_SIMPLE_AUTOMATION_TEST(TExportRequestSerializeOneTestError, "Plugin.CaptureProtocolStack.Export.ExportRequest.SerializeOne.Error", EAutomationTestFlags::ApplicationContextMask | EAutomationTestFlags::SmokeFilter)
+bool TExportRequestSerializeOneTestError::RunTest(const FString& InParameters)
+{
+	// Prepare test
+	const FString TakeName = TEXT("TakeName");
+	const FString FileName = TEXT("FileName");
+	const uint16 TakeNameLen = TakeName.Len();
+	const uint16 FileNameLen = FileName.Len();
+
+	const uint64 Offset = 0;
+
+	TArray<uint8> Packet;
+	Packet.Append(reinterpret_cast<const uint8*>(&TakeNameLen), sizeof(TakeNameLen));
+	Packet.Append(reinterpret_cast<const uint8*>(TCHAR_TO_UTF8(*TakeName)), TakeName.Len());
+	Packet.Append(reinterpret_cast<const uint8*>(&FileNameLen), sizeof(FileNameLen));
+	Packet.Append(reinterpret_cast<const uint8*>(TCHAR_TO_UTF8(*FileName)), FileName.Len());
+	Packet.Append(reinterpret_cast<const uint8*>(&Offset), sizeof(Offset));
+
+	FFailedDataSender DataSender;
+
+	FExportRequest Request(TakeName, FileName, Offset);
+
+	// Serialization
+	TProtocolResult<void> ExportRequestResult = FExportRequest::Serialize(Request, DataSender);
+	TestTrue(TEXT("Request Serialize"), ExportRequestResult.HasError());
+
+	return true;
+}
+
+#endif // WITH_DEV_AUTOMATION_TESTS
\ No newline at end of file
diff --git a/Engine/Plugins/VirtualProduction/CaptureManager/Source/CaptureProtocolStack/Private/Tests/Export/TestExportResponse.cpp b/Engine/Plugins/VirtualProduction/CaptureManager/Source/CaptureProtocolStack/Private/Tests/Export/TestExportResponse.cpp
new file mode 100644
--- /dev/null
+++ b/Engine/Plugins/VirtualProduction/CaptureManager/Source/CaptureProtocolStack/Private/Tests/Export/TestExportResponse.cpp
@@ -0,0 +1,108 @@
+// Copyright Epic Games, Inc. All Rights Reserved.
+
+#include "ExportClient/Messages/ExportResponse.h"
+
+#include "Misc/AutomationTest.h"
+
+#include "Tests/Utility.h"
+
+#if WITH_DEV_AUTOMATION_TESTS
+
+TStaticArray<uint8, 16> CreateHash()
+{
+	TStaticArray<uint8, 16> StaticArray;
+	for (int32 Index = 0; Index < StaticArray.Num(); ++Index)
+	{
+		StaticArray[Index] = Index;
+	}
+
+	return StaticArray;
+}
+
+IMPLEMENT_SIMPLE_AUTOMATION_TEST(TExportResponseDeserializeOneTest, "Plugin.CaptureProtocolStack.Export.ExportResponse.DeserializeOne.Success", EAutomationTestFlags::ApplicationContextMask | EAutomationTestFlags::SmokeFilter)
+bool TExportResponseDeserializeOneTest::RunTest(const FString& InParameters)
+{
+	// Prepare test
+	const FExportResponse::EStatus Status = FExportResponse::EStatus::Success;
+	const uint64 Length = 50;
+
+	TArray<uint8> Packet;
+	Packet.Append(reinterpret_cast<const uint8*>(&Status), sizeof(Status));
+	Packet.Append(reinterpret_cast<const uint8*>(&Length), sizeof(Length));
+
+	FDataProvider DataProvider(MoveTemp(Packet));
+
+	// Deserialization
+	TProtocolResult<FExportResponse> ExportResponseResult = FExportResponse::Deserialize(DataProvider);
+	TestTrue(TEXT("Response Deserialize"), ExportResponseResult.HasValue());
+
+	FExportResponse ExportResponse = ExportResponseResult.StealValue();
+
+	TestEqual(TEXT("Response Deserialize"), ExportResponse.GetStatus(), Status);
+	TestEqual(TEXT("Response Deserialize"), ExportResponse.GetLength(), Length);
+
+	return true;
+}
+
+IMPLEMENT_SIMPLE_AUTOMATION_TEST(TExportResponseDeserializeOneTestInvalidSize, "Plugin.CaptureProtocolStack.Export.ExportResponse.DeserializeOne.InvalidSize", EAutomationTestFlags::ApplicationContextMask | EAutomationTestFlags::SmokeFilter)
+bool TExportResponseDeserializeOneTestInvalidSize::RunTest(const FString& InParameters)
+{
+	// Prepare test
+	TArray<uint8> Packet;
+
+	FDataProvider DataProvider(MoveTemp(Packet));
+
+	// Deserialization
+	TProtocolResult<FExportResponse> ExportResponseResult = FExportResponse::Deserialize(DataProvider);
+	TestTrue(TEXT("Response Deserialize"), ExportResponseResult.HasError());
+
+	return true;
+}
+
+IMPLEMENT_SIMPLE_AUTOMATION_TEST(TExportResponseSerializeOneTest, "Plugin.CaptureProtocolStack.Export.ExportResponse.SerializeOne.Success", EAutomationTestFlags::ApplicationContextMask | EAutomationTestFlags::SmokeFilter)
+bool TExportResponseSerializeOneTest::RunTest(const FString& InParameters)
+{
+	// Prepare test
+	const FExportResponse::EStatus Status = FExportResponse::EStatus::Success;
+	const uint64 Length = 50;
+
+	TArray<uint8> Packet;
+	Packet.Append(reinterpret_cast<const uint8*>(&Status), sizeof(Status));
+	Packet.Append(reinterpret_cast<const uint8*>(&Length), sizeof(Length));
+
+	FDataSender DataSender;
+
+	FExportResponse Response(Status, Length);
+
+	// Serialization
+	TProtocolResult<void> ExportResponseResult = FExportResponse::Serialize(Response, DataSender);
+	TestTrue(TEXT("Response Serialize"), ExportResponseResult.HasValue());
+
+	TestEqual(TEXT("Response Serialize"), DataSender.GetData(), Packet);
+
+	return true;
+}
+
+IMPLEMENT_SIMPLE_AUTOMATION_TEST(TExportResponseSerializeOneTestError, "Plugin.CaptureProtocolStack.Export.ExportResponse.SerializeOne.Error", EAutomationTestFlags::ApplicationContextMask | EAutomationTestFlags::SmokeFilter)
+bool TExportResponseSerializeOneTestError::RunTest(const FString& InParameters)
+{
+	// Prepare test
+	const FExportResponse::EStatus Status = FExportResponse::EStatus::Success;
+	const uint64 Length = 50;
+
+	TArray<uint8> Packet;
+	Packet.Append(reinterpret_cast<const uint8*>(&Status), sizeof(Status));
+	Packet.Append(reinterpret_cast<const uint8*>(&Length), sizeof(Length));
+
+	FFailedDataSender DataSender;
+
+	FExportResponse Response(Status, Length);
+
+	// Serialization
+	TProtocolResult<void> ExportResponseResult = FExportResponse::Serialize(Response, DataSender);
+	TestTrue(TEXT("Response Serialize"), ExportResponseResult.HasError());
+
+	return true;
+}
+
+#endif // WITH_DEV_AUTOMATION_TESTS
\ No newline at end of file
diff --git a/Engine/Plugins/VirtualProduction/CaptureManager/Source/CaptureProtocolStack/Private/Tests/Utility.h b/Engine/Plugins/VirtualProduction/CaptureManager/Source/CaptureProtocolStack/Private/Tests/Utility.h
new file mode 100644
--- /dev/null
+++ b/Engine/Plugins/VirtualProduction/CaptureManager/Source/CaptureProtocolStack/Private/Tests/Utility.h
@@ -0,0 +1,68 @@
+// Copyright Epic Games, Inc. All Rights Reserved.
+
+#pragma once
+
+#include "Communication/TcpReaderWriter.h"
+
+class FDataProvider : public ITcpSocketReader
+{
+public:
+
+	FDataProvider(TArray<uint8> InData)
+		: Data(MoveTemp(InData))
+	{
+	}
+
+	virtual TProtocolResult<TArray<uint8>> ReceiveMessage(const uint64 InSize, const uint32 InMaxWaitTimeMs = DefaultWaitTimeoutMs) override
+	{
+		if (InSize > Data.Num())
+		{
+			return FCaptureProtocolError(TEXT("Failed to receive data"));
+		}
+
+		TArray<uint8> DataToBeReturned(Data.GetData(), InSize);
+
+		Data.RemoveAt(0, InSize);
+
+		return DataToBeReturned;
+	}
+
+private:
+
+	TArray<uint8> Data;
+};
+
+class FDataSender : public ITcpSocketWriter
+{
+public:
+
+	FDataSender()
+	{
+	}
+
+	virtual TProtocolResult<void> SendMessage(const TArray<uint8>& InPayload) override
+	{
+		Data.Append(InPayload.GetData(), InPayload.Num());
+
+		return ResultOk;
+	}
+
+	const TArray<uint8>& GetData() const
+	{
+		return Data;
+	}
+
+private:
+
+	TArray<uint8> Data;
+};
+
+class FFailedDataSender : public FDataSender
+{
+public:
+
+	virtual TProtocolResult<void> SendMessage(const TArray<uint8>& InPayload) override
+	{
+		return FCaptureProtocolError(TEXT("Failed to send the data"));
+	}
+};
\ No newline at end of file
diff --git a/Engine/Plugins/VirtualProduction/CaptureManager/Source/CaptureProtocolStack/Private/Utility/CPSTimerManager.cpp b/Engine/Plugins/VirtualProduction/CaptureManager/Source/CaptureProtocolStack/Private/Utility/CPSTimerManager.cpp
new file mode 100644
--- /dev/null
+++ b/Engine/Plugins/VirtualProduction/CaptureManager/Source/CaptureProtocolStack/Private/Utility/CPSTimerManager.cpp
@@ -0,0 +1,70 @@
+// Copyright Epic Games, Inc. All Rights Reserved.
+
+#include "Utility/CPSTimerManager.h"
+
+const float FCPSTimerManager::IdealTimeResolution = 0.1f;
+
+FCPSTimerManager::FCPSTimerManager(const float Resolution)
+{
+	Runnable = MakeUnique<FTimerManagerRunnable>(this, Resolution);
+	Thread.Reset(FRunnableThread::Create(Runnable.Get(), TEXT("Timer Manager"), 128 * 1024, TPri_Normal, FPlatformAffinity::GetPoolThreadMask()));
+}
+
+FCPSTimerManager::~FCPSTimerManager()
+{
+	Thread->Kill(true);
+
+	Thread = nullptr;
+	Runnable = nullptr;
+}
+
+FCPSTimerManager::FTimerHandle FCPSTimerManager::AddTimer(FTimerDelegate InDelegate, float InRate, bool InbLoop, float InFirstDelay)
+{
+	FCPSTimerManager::FTimerHandle Handle = Ticker.AddTicker(TEXT("Timer"), InRate, [InbLoop, Delegate = MoveTemp(InDelegate)](float InDelay)
+	{
+		Delegate.Execute();
+
+		return InbLoop;
+	});
+
+	Handle.Pin()->FireTime -= InRate - InFirstDelay;
+
+	return Handle;
+}
+
+void FCPSTimerManager::RemoveTimer(FTimerHandle Handle)
+{
+	Ticker.RemoveTicker(MoveTemp(Handle));
+}
+
+FCPSTimerManager::FTimerManagerRunnable::FTimerManagerRunnable(FCPSTimerManager* InOwner, const float InResolution)
+	: Owner(InOwner)
+	, IdealTimeResolution(InResolution)
+{
+}
+
+uint32 FCPSTimerManager::FTimerManagerRunnable::Run()
+{
+	double LastTime = FPlatformTime::Seconds();
+	
+	while (bShouldRun.load())
+	{
+		const double CurrentTime = FPlatformTime::Seconds();
+		const double DeltaTime = CurrentTime - LastTime;
+
+		Owner->Ticker.Tick(DeltaTime);
+
+		FPlatformProcess::Sleep(FMath::Max<float>(0.0f, IdealTimeResolution - (FPlatformTime::Seconds() - LastTime)));
+
+		LastTime = CurrentTime;
+	}
+
+	Owner->Ticker.Reset();
+
+	return 0;
+}
+
+void FCPSTimerManager::FTimerManagerRunnable::Stop()
+{
+	bShouldRun.store(false);
+}
\ No newline at end of file
diff --git a/Engine/Plugins/VirtualProduction/CaptureManager/Source/CaptureProtocolStack/Private/Utility/Error.cpp b/Engine/Plugins/VirtualProduction/CaptureManager/Source/CaptureProtocolStack/Private/Utility/Error.cpp
new file mode 100644
--- /dev/null
+++ b/Engine/Plugins/VirtualProduction/CaptureManager/Source/CaptureProtocolStack/Private/Utility/Error.cpp
@@ -0,0 +1,25 @@
+// Copyright Epic Games, Inc. All Rights Reserved.
+
+#include "Utility/Error.h"
+
+FCaptureProtocolError::FCaptureProtocolError()
+    : Message("")
+    , Code(0)
+{
+}
+
+FCaptureProtocolError::FCaptureProtocolError(FString InMessage, int32 InCode)
+    : Message(MoveTemp(InMessage))
+    , Code(InCode)
+{
+}
+
+const FString& FCaptureProtocolError::GetMessage() const
+{
+    return Message;
+}
+
+int32 FCaptureProtocolError::GetCode() const
+{
+    return Code;
+}
\ No newline at end of file
diff --git a/Engine/Plugins/VirtualProduction/CaptureManager/Source/CaptureProtocolStack/Public/CaptureProtocolStack.h b/Engine/Plugins/VirtualProduction/CaptureManager/Source/CaptureProtocolStack/Public/CaptureProtocolStack.h
new file mode 100644
--- /dev/null
+++ b/Engine/Plugins/VirtualProduction/CaptureManager/Source/CaptureProtocolStack/Public/CaptureProtocolStack.h
@@ -0,0 +1,22 @@
+// Copyright Epic Games, Inc. All Rights Reserved.
+
+#pragma once
+
+#include "Modules/ModuleManager.h"
+
+#include "Utility/CPSTimerManager.h"
+
+class CAPTUREPROTOCOLSTACK_API FCaptureProtocolStackModule : public IModuleInterface
+{
+public:
+
+	/** IModuleInterface implementation */
+	virtual void StartupModule() override;
+	virtual void ShutdownModule() override;
+
+	FCPSTimerManager& GetTimerManager();
+
+private:
+
+	TUniquePtr<FCPSTimerManager> TimerManager;
+};
diff --git a/Engine/Plugins/VirtualProduction/CaptureManager/Source/CaptureProtocolStack/Public/Communication/TcpClient.h b/Engine/Plugins/VirtualProduction/CaptureManager/Source/CaptureProtocolStack/Public/Communication/TcpClient.h
new file mode 100644
--- /dev/null
+++ b/Engine/Plugins/VirtualProduction/CaptureManager/Source/CaptureProtocolStack/Public/Communication/TcpClient.h
@@ -0,0 +1,74 @@
+// Copyright Epic Games, Inc. All Rights Reserved.
+
+#pragma once
+
+#include "TcpReaderWriter.h"
+
+#include "Utility/Error.h"
+
+#include "Runtime/Sockets/Public/IPAddress.h"
+
+#include "Common/TcpSocketBuilder.h"
+
+#include "EngineLogs.h"
+#include "Engine/EngineTypes.h"
+#include "Engine/EngineBaseTypes.h"
+
+#include "Containers/UnrealString.h"
+
+#include "Templates/UniquePtr.h"
+
+#include "Misc/DateTime.h"
+
+using FSocketPtr = TUniquePtr<FSocket, FSocketDeleter>;
+
+class FTcpClient final
+{
+public:
+
+    static constexpr uint32 BufferSize = 2 * 1024 * 1024;
+    static constexpr int32 DisconnectedError = -10;
+	static constexpr int32 TimeoutError = -1;
+
+    FTcpClient();
+    ~FTcpClient();
+
+	TProtocolResult<void> Init();
+    // Blocking function until connection is established
+	TProtocolResult<void> Start(const FString& InServerAddress);
+	TProtocolResult<void> Stop();
+
+	bool IsRunning() const;
+
+	TProtocolResult<void> SendMessage(const TArray<uint8>& InPayload);
+	TProtocolResult<TArray<uint8>> ReceiveMessage(const uint64 InSize, const uint32 InWaitTimeoutMs = ITcpSocketReader::DefaultWaitTimeoutMs);
+
+private:
+
+    FSocketPtr TcpSocket;
+
+    bool bRunning;
+};
+
+class FTcpClientReader final : public ITcpSocketReader
+{
+public:
+    FTcpClientReader(FTcpClient& InClient);
+
+    virtual TProtocolResult<TArray<uint8>> ReceiveMessage(const uint64 InSize, const uint32 InWaitTimeoutMs = DefaultWaitTimeoutMs) override;
+
+private:
+    FTcpClient& Client;
+};
+
+class FTcpClientWriter final : public ITcpSocketWriter
+{
+public:
+    FTcpClientWriter(FTcpClient& InClient);
+
+    virtual TProtocolResult<void> SendMessage(const TArray<uint8>& InPayload) override;
+
+private:
+
+    FTcpClient& Client;
+};
diff --git a/Engine/Plugins/VirtualProduction/CaptureManager/Source/CaptureProtocolStack/Public/Communication/TcpReaderWriter.h b/Engine/Plugins/VirtualProduction/CaptureManager/Source/CaptureProtocolStack/Public/Communication/TcpReaderWriter.h
new file mode 100644
--- /dev/null
+++ b/Engine/Plugins/VirtualProduction/CaptureManager/Source/CaptureProtocolStack/Public/Communication/TcpReaderWriter.h
@@ -0,0 +1,25 @@
+// Copyright Epic Games, Inc. All Rights Reserved.
+
+#pragma once
+
+#include "Utility/Error.h"
+
+#include "Containers/Array.h"
+
+class ITcpSocketReader
+{
+public:
+	static constexpr uint32 DefaultWaitTimeoutMs = 1000;
+
+    virtual ~ITcpSocketReader() = default;
+
+    virtual TProtocolResult<TArray<uint8>> ReceiveMessage(const uint64 InSize, const uint32 InWaitTimeoutMs = DefaultWaitTimeoutMs) = 0;
+};
+
+class ITcpSocketWriter
+{
+public:
+    virtual ~ITcpSocketWriter() = default;
+
+    virtual TProtocolResult<void> SendMessage(const TArray<uint8>& InPayload) = 0;
+};
\ No newline at end of file
diff --git a/Engine/Plugins/VirtualProduction/CaptureManager/Source/CaptureProtocolStack/Public/Communication/TcpServer.h b/Engine/Plugins/VirtualProduction/CaptureManager/Source/CaptureProtocolStack/Public/Communication/TcpServer.h
new file mode 100644
--- /dev/null
+++ b/Engine/Plugins/VirtualProduction/CaptureManager/Source/CaptureProtocolStack/Public/Communication/TcpServer.h
@@ -0,0 +1,109 @@
+// Copyright Epic Games, Inc. All Rights Reserved.
+
+#pragma once
+
+#include "Communication/TcpReaderWriter.h"
+
+#include "Utility/Error.h"
+
+#include "Common/TcpSocketBuilder.h"
+#include "Common/TcpListener.h"
+
+#include "EngineLogs.h"
+#include "Engine/EngineTypes.h"
+#include "Engine/EngineBaseTypes.h"
+#include "Async/Async.h"
+
+#include "Containers/UnrealString.h"
+
+#include "Templates/SharedPointer.h"
+#include "Templates/UniquePtr.h"
+
+#include "Delegates/Delegate.h"
+
+using FSocketPtr = TUniquePtr<FSocket, FSocketDeleter>;
+
+class CAPTUREPROTOCOLSTACK_API FTcpClientHandler final
+{
+public:
+
+    static constexpr uint32 MaxBufferSize = 500 * 1024;
+    static constexpr int32 DisconnectedError = -10;
+	static constexpr int32 TimeoutError = -1;
+
+    FTcpClientHandler(FSocketPtr InSocket, FString InEndpoint);
+    ~FTcpClientHandler();
+
+	TProtocolResult<void> SendMessage(const TArray<uint8>& InData);
+	TProtocolResult<TArray<uint8>> ReceiveMessage(const uint64 InSize, const uint32 InWaitTimeoutMs = ITcpSocketReader::DefaultWaitTimeoutMs);
+
+	TProtocolResult<uint32> HasPendingData() const;
+
+    const FString& GetEndpoint() const;
+
+    bool operator==(const FTcpClientHandler& InOther);
+
+private:
+
+    FSocketPtr Socket;
+    FString Endpoint;
+};
+
+class CAPTUREPROTOCOLSTACK_API FTcpServer
+{
+public:
+
+    static constexpr uint32 ThreadWaitTime = 500; // Milliseconds
+	static constexpr uint16 AnyPort = 0; // OS assigned port
+
+    DECLARE_DELEGATE_TwoParams(FConnectionHandler, TWeakPtr<FTcpClientHandler> InClient, bool bConnected)
+
+    FTcpServer(const uint32 InMaxNumberOfClients);
+
+	TProtocolResult<uint16> Start(const uint16 InListenPort = AnyPort);
+	TProtocolResult<void> Stop();
+
+	TProtocolResult<void> SendMessage(const TArray<uint8>& InMessage, const FString& InEndpoint);
+
+    void DisconnectClient(const FString& InEndpoint);
+    void DisconnectClientAsync(const FString& InEndpoint);
+    void SetConnectionHandler(FConnectionHandler InOnConnectionHandler);
+
+	int32 GetPort() const;
+
+private:
+
+    TUniquePtr<FTcpListener> Listener;
+	FSocketPtr Socket;
+
+    uint32 MaxNumberOfClients;
+    TMap<FString, TSharedPtr<FTcpClientHandler>> Clients;
+
+    FConnectionHandler OnConnectionHandler;
+    FRWLock Mutex;
+
+    bool bRunning;
+};
+
+class FTcpConnectionReader final : public ITcpSocketReader
+{
+public:
+    FTcpConnectionReader(FTcpClientHandler& InClient);
+
+    virtual TProtocolResult<TArray<uint8>> ReceiveMessage(const uint64 InSize, const uint32 InWaitTimeoutMs = DefaultWaitTimeoutMs) override;
+
+private:
+    FTcpClientHandler& Client;
+};
+
+class FTcpConnectionWriter final : public ITcpSocketWriter
+{
+public:
+    FTcpConnectionWriter(FTcpClientHandler& InClient);
+
+    virtual TProtocolResult<void> SendMessage(const TArray<uint8>& InPayload) override;
+
+private:
+
+    FTcpClientHandler& Client;
+};
diff --git a/Engine/Plugins/VirtualProduction/CaptureManager/Source/CaptureProtocolStack/Public/Communication/UdpClient.h b/Engine/Plugins/VirtualProduction/CaptureManager/Source/CaptureProtocolStack/Public/Communication/UdpClient.h
new file mode 100644
--- /dev/null
+++ b/Engine/Plugins/VirtualProduction/CaptureManager/Source/CaptureProtocolStack/Public/Communication/UdpClient.h
@@ -0,0 +1,52 @@
+// Copyright Epic Games, Inc. All Rights Reserved.
+
+#pragma once
+
+#include "Utility/Error.h"
+
+#include "Runtime/Sockets/Public/IPAddress.h"
+
+#include "Common/UdpSocketReceiver.h"
+#include "Common/UdpSocketSender.h"
+#include "Common/UdpSocketBuilder.h"
+
+#include "EngineLogs.h"
+#include "Engine/EngineTypes.h"
+#include "Engine/EngineBaseTypes.h"
+
+#include "Containers/UnrealString.h"
+
+#include "Templates/UniquePtr.h"
+
+struct FUdpClientConfigure
+{
+    uint16 ListenPort;
+    FString MulticastIpAddress = "";
+};
+
+#define CHECK_BOOL(Function) if (bool Result = Function; !Result) { return FCaptureProtocolError("Failed to execute function"); }
+
+class FUdpClient final
+{
+public:
+
+    static constexpr uint32 ThreadWaitTime = 500; // Milliseconds
+    static constexpr uint32 BufferSize = 2 * 1024 * 1024;
+
+    FUdpClient();
+    ~FUdpClient();
+
+    TProtocolResult<void> Init(FUdpClientConfigure InConfig, FOnSocketDataReceived InReceiveHandler);
+	TProtocolResult<void> Start();
+	TProtocolResult<void> Stop();
+
+	TProtocolResult<int32> SendMessage(const TArray<uint8>& InPayload, const FString& InEndpoint);
+
+private:
+
+    TUniquePtr<FSocket> UdpSocket;
+
+    TUniquePtr<FUdpSocketReceiver> UdpReceiver;
+
+    bool bRunning;
+};
diff --git a/Engine/Plugins/VirtualProduction/CaptureManager/Source/CaptureProtocolStack/Public/Control/Communication/ControlCommunication.h b/Engine/Plugins/VirtualProduction/CaptureManager/Source/CaptureProtocolStack/Public/Control/Communication/ControlCommunication.h
new file mode 100644
--- /dev/null
+++ b/Engine/Plugins/VirtualProduction/CaptureManager/Source/CaptureProtocolStack/Public/Control/Communication/ControlCommunication.h
@@ -0,0 +1,72 @@
+// Copyright Epic Games, Inc. All Rights Reserved.
+
+#pragma once
+
+#include "ControlPacket.h"
+
+#include "Communication/TcpClient.h"
+
+#include "Utility/QueueRunner.h"
+#include "Utility/Error.h"
+
+#include "HAL/Runnable.h"
+#include "HAL/RunnableThread.h"
+#include "HAL/Event.h"
+
+class FControlCommunication final
+{
+public:
+
+    DECLARE_DELEGATE_OneParam(FOnPacketReceived, FControlPacket InPacket)
+
+    using FRunnerType = TQueueRunner<FControlPacket>;
+
+    FControlCommunication();
+
+	TProtocolResult<void> Init();
+	TProtocolResult<void> Start(const FString& InServerIp, const uint16 InServerPort);
+	TProtocolResult<void> Stop();
+
+	bool IsRunning() const;
+
+    void SendMessage(FControlPacket InMessage);
+
+    void SetReceiveHandler(FOnPacketReceived InReceiveHandler);
+
+private:
+	
+	class FCommunicationRunnable : 
+		public FRunnable
+	{
+	public:
+		FCommunicationRunnable(FControlCommunication& InCommunication);
+		virtual ~FCommunicationRunnable();
+
+		uint32 Run() override;
+		void Stop() override;
+
+		void Join();
+
+	private:
+
+		FControlCommunication& Communication;
+		bool bIsRunning;
+		FEvent* DoneEvent;
+	};
+
+
+    void OnProcessReceivedPacket(FControlPacket InMessage);
+    void OnProcessSentPacket(FControlPacket InMessage);
+	void RunnableHandler(TSharedPtr<FRunnable> InRunnable);
+
+	TProtocolResult<FControlPacketHeader> ReceiveControlHeader();
+	TProtocolResult<void> ReceiveControlPacket(const FControlPacketHeader& InHeader);
+
+    FTcpClient Client;
+    FOnPacketReceived ReceiveHandler;
+
+    FRunnerType SynchronizedReceiver;
+    FRunnerType SynchronizedSender;
+	TQueueRunner<TSharedPtr<FRunnable>> SynchronizedRunnable;
+	TWeakPtr<FCommunicationRunnable> Runnable;
+};
\ No newline at end of file
diff --git a/Engine/Plugins/VirtualProduction/CaptureManager/Source/CaptureProtocolStack/Public/Control/Communication/ControlPacket.h b/Engine/Plugins/VirtualProduction/CaptureManager/Source/CaptureProtocolStack/Public/Control/Communication/ControlPacket.h
new file mode 100644
--- /dev/null
+++ b/Engine/Plugins/VirtualProduction/CaptureManager/Source/CaptureProtocolStack/Public/Control/Communication/ControlPacket.h
@@ -0,0 +1,58 @@
+// Copyright Epic Games, Inc. All Rights Reserved.
+
+#pragma once
+
+#include "Communication/TcpReaderWriter.h"
+
+#include "Utility/Error.h"
+
+#include "Containers/Array.h"
+
+#include "Dom/JsonObject.h"
+
+class FControlPacketHeader final
+{
+public:
+
+    static const TArray<uint8> Header;
+
+    static TProtocolResult<FControlPacketHeader> Deserialize(ITcpSocketReader& InReader);
+    static uint32 GetHeaderSize();
+
+    FControlPacketHeader();
+	FControlPacketHeader(uint16 InVersion, uint32 InPayloadSize);
+
+	uint16 GetVersion() const;
+    uint32 GetPayloadSize() const;
+
+private:
+	
+	uint16 Version;
+    uint32 PayloadSize;
+};
+
+class FControlPacket final
+{
+public:
+
+    static TProtocolResult<FControlPacket> Deserialize(const FControlPacketHeader& InPacketHeader, ITcpSocketReader& InReader);
+    static TProtocolResult<void> Serialize(const FControlPacket& InMessage, ITcpSocketWriter& InWriter);
+
+    FControlPacket();
+    FControlPacket(uint16 InVersion, TArray<uint8> InPayload);
+
+    FControlPacket(const FControlPacket& InOther) = default;
+    FControlPacket(FControlPacket&& InOther) = default;
+
+    FControlPacket& operator=(const FControlPacket& InOther) = default;
+    FControlPacket& operator=(FControlPacket&& InOther) = default;
+
+	uint16 GetVersion() const;
+    const TArray<uint8>& GetPayload() const;
+    uint32 GetPayloadSize() const;
+
+private:
+
+	FControlPacketHeader Header;
+    TArray<uint8> Payload;
+};
diff --git a/Engine/Plugins/VirtualProduction/CaptureManager/Source/CaptureProtocolStack/Public/Control/ControlMessenger.h b/Engine/Plugins/VirtualProduction/CaptureManager/Source/CaptureProtocolStack/Public/Control/ControlMessenger.h
new file mode 100644
--- /dev/null
+++ b/Engine/Plugins/VirtualProduction/CaptureManager/Source/CaptureProtocolStack/Public/Control/ControlMessenger.h
@@ -0,0 +1,196 @@
+// Copyright Epic Games, Inc. All Rights Reserved.
+
+#pragma once
+
+#include "Utility/Definitions.h"
+
+#include "Control/Communication/ControlCommunication.h"
+#include "Control/Communication/ControlPacket.h"
+
+#include "Messages/ControlMessage.h"
+#include "Messages/ControlRequest.h"
+#include "Messages/ControlResponse.h"
+#include "Messages/ControlUpdate.h"
+
+#include "Utility/Error.h"
+#include "Utility/CPSTimerManager.h"
+
+#include "Misc/ScopedEvent.h"
+
+#include "Async/Async.h"
+
+class FKeepAliveCounter
+{
+public:
+	FKeepAliveCounter();
+
+	void Increment();
+	void Reset();
+	bool HasReached(uint16 InBound);
+
+private:
+
+	std::atomic<uint16> Counter;
+};
+
+class CAPTUREPROTOCOLSTACK_API FControlMessenger final
+{
+public:
+
+    template<class RequestType>
+    using FOnControlResponse = TDelegate<void(TProtocolResult<typename RequestType::ResponseType> InResult)>;
+
+	DECLARE_DELEGATE_OneParam(FOnDisconnect, const FString& InCause);
+
+    static constexpr int32 ResponseWaitTime = 3; // Seconds
+    static constexpr int32 KeepAliveInterval = 5; // Seconds
+    static const TCHAR* HandshakeSessionId;
+
+    FControlMessenger();
+    ~FControlMessenger();
+    
+    void RegisterUpdateHandler(FString InAddressPath, FControlUpdate::FOnUpdateMessage InUpdateHandler);
+	void RegisterDisconnectHandler(FOnDisconnect InOnDisconnectHandler);
+
+	TProtocolResult<void> Start(const FString& InServerIp, const uint16 InServerPort);
+	TProtocolResult<void> Stop();
+
+	TProtocolResult<void> StartSession();
+    TProtocolResult<FGetServerInformationResponse> GetServerInformation();
+
+    template<class RequestType>
+	TProtocolResult<typename RequestType::ResponseType> SendRequest(RequestType InRequest)
+    {
+        FControlMessage Message(InRequest.GetAddressPath(), FControlMessage::EType::Request, InRequest.GetBody());
+
+        uint32 TransactionId = GenerateTransactionId();
+
+        {
+            FScopeLock SessionLock(&SessionIdMutex);
+            Message.SetSessionId(SessionId);
+        }
+        
+        Message.SetTransactionId(TransactionId);
+        Message.SetTimestamp(GetTimestamp());
+        
+		TProtocolResult<FControlPacket> SerializeResult = FControlMessage::Serialize(Message);
+
+        if (SerializeResult.HasError())
+        {
+            return FCaptureProtocolError(TEXT("Failed to serialize a request."));
+        }
+
+        TPromise<TProtocolResult<FControlMessage>> Promise;
+        TFuture<TProtocolResult<FControlMessage>> Future = Promise.GetFuture();
+
+        {
+            FScopeLock Lock(&RequestsMutex);
+            RequestContexts.Emplace(TransactionId, MakeUnique<FRequestContext>(MoveTemp(Message), MoveTemp(Promise)));
+        }
+
+        SendPacket(SerializeResult.StealValue());
+
+        if (!Future.WaitFor(FTimespan::FromSeconds(ResponseWaitTime)))
+        {
+            FScopeLock Lock(&RequestsMutex);
+
+            const TUniquePtr<FRequestContext>& RequestContext = RequestContexts[TransactionId];
+
+			RequestContext->Promise.SetValue(FCaptureProtocolError("Broken promise")); // Actual error to avoid the destructor error.
+
+            RequestContexts.Remove(TransactionId);
+
+            return FCaptureProtocolError(TEXT("Server failed to respond within 3 seconds."));
+        }
+
+        FScopeLock Lock(&RequestsMutex);
+        const TUniquePtr<FRequestContext>& RequestContext = RequestContexts[TransactionId];
+
+        const TProtocolResult<FControlMessage>& FutureResult = Future.Get();
+        check(FutureResult.HasValue());
+
+        const FControlMessage& ResponseMessage = FutureResult.GetValue();
+
+        {
+            FScopeLock SessionLock(&SessionIdMutex);
+            if (ResponseMessage.GetSessionId() != SessionId)
+            {
+                return FCaptureProtocolError(TEXT("Invalid session ID arrived"));
+            }
+        }
+
+        if (!ResponseMessage.GetErrorName().IsEmpty())
+        {
+            return FCaptureProtocolError(TEXT("Server respond with error: ") + ResponseMessage.GetErrorName());
+        }
+
+        typename RequestType::ResponseType Response;
+        RequestContexts.Remove(TransactionId);
+
+		TProtocolResult<void> ParseResult = Response.Parse(ResponseMessage.GetBody());
+        if (ParseResult.HasError())
+        {
+            return FCaptureProtocolError(TEXT("Failed to parse the response: ") + ParseResult.StealError().GetMessage());
+        }
+        
+        return Response;
+    }
+
+    template<class RequestType>
+    void SendAsyncRequest(RequestType InRequest, FOnControlResponse<RequestType> InOnResponse)
+    {
+		AsyncRequestRunner.Add(FAsyncRequestDelegate::CreateLambda([this, Request(MoveTemp(InRequest)), OnResponse(MoveTemp(InOnResponse))]() mutable
+		{
+			TProtocolResult<typename RequestType::ResponseType> Result = SendRequest(MoveTemp(Request));
+
+			OnResponse.ExecuteIfBound(MoveTemp(Result));
+		}));
+    }
+
+private:
+	DECLARE_DELEGATE(FAsyncRequestDelegate)
+
+    struct FRequestContext
+    {
+        FRequestContext(FControlMessage InRequest, TPromise<TProtocolResult<FControlMessage>> InPromise)
+            : Request(MoveTemp(InRequest))
+            , Promise(MoveTemp(InPromise))
+        {
+        }
+
+        FControlMessage Request;
+
+        TPromise<TProtocolResult<FControlMessage>> Promise;
+    };
+
+    void SendPacket(FControlPacket InPacket);
+    void KeepAlive();
+    void MessageHandler(FControlPacket InPacket);
+
+	uint32 GenerateTransactionId() const;
+	uint64 GetTimestamp() const;
+
+    void StartKeepAliveTimer();
+    void StopKeepAliveTimer();
+
+	void OnAsyncRequestProcess(FAsyncRequestDelegate InAsyncDelegate);
+
+    FControlCommunication Communication;
+
+    FCriticalSection SessionIdMutex;
+    FString SessionId;
+
+    FCriticalSection RequestsMutex;
+    TMap<uint32, TUniquePtr<FRequestContext>> RequestContexts;
+
+    FCriticalSection UpdatesMutex;
+    TMap<FString, FControlUpdate::FOnUpdateMessage> UpdateHandlers;
+
+	FCPSTimerManager::FTimerHandle KeepAliveTimer;
+	FKeepAliveCounter KeepAliveFailures;
+
+	TQueueRunner<FAsyncRequestDelegate> AsyncRequestRunner;
+	FOnDisconnect OnDisconnectHandler;
+
+	FRandomStream RandomStream;
+};
\ No newline at end of file
diff --git a/Engine/Plugins/VirtualProduction/CaptureManager/Source/CaptureProtocolStack/Public/Control/Messages/Constants.h b/Engine/Plugins/VirtualProduction/CaptureManager/Source/CaptureProtocolStack/Public/Control/Messages/Constants.h
new file mode 100644
--- /dev/null
+++ b/Engine/Plugins/VirtualProduction/CaptureManager/Source/CaptureProtocolStack/Public/Control/Messages/Constants.h
@@ -0,0 +1,130 @@
+// Copyright Epic Games, Inc. All Rights Reserved.
+
+#pragma once
+
+#include "Containers/UnrealString.h"
+
+namespace UE::CPS::AddressPaths
+{
+// Requests/Responses
+extern CAPTUREPROTOCOLSTACK_API const TCHAR* GKeepAlive;
+extern CAPTUREPROTOCOLSTACK_API const TCHAR* GStartSession;
+extern CAPTUREPROTOCOLSTACK_API const TCHAR* GStopSession;
+extern CAPTUREPROTOCOLSTACK_API const TCHAR* GSubscribe;
+extern CAPTUREPROTOCOLSTACK_API const TCHAR* GUnsubscribe;
+extern CAPTUREPROTOCOLSTACK_API const TCHAR* GGetServerInformation;
+extern CAPTUREPROTOCOLSTACK_API const TCHAR* GGetState;
+extern CAPTUREPROTOCOLSTACK_API const TCHAR* GStartRecordingTake;
+extern CAPTUREPROTOCOLSTACK_API const TCHAR* GStopRecordingTake;
+extern CAPTUREPROTOCOLSTACK_API const TCHAR* GAbortRecordingTake;
+extern CAPTUREPROTOCOLSTACK_API const TCHAR* GGetTakeList;
+extern CAPTUREPROTOCOLSTACK_API const TCHAR* GGetTakeMetadata;
+extern CAPTUREPROTOCOLSTACK_API const TCHAR* GExportTakeData;
+extern CAPTUREPROTOCOLSTACK_API const TCHAR* GExportTakeVideoFrame;
+extern CAPTUREPROTOCOLSTACK_API const TCHAR* GExportCameraFeedFrame;
+extern CAPTUREPROTOCOLSTACK_API const TCHAR* GPauseExport;
+extern CAPTUREPROTOCOLSTACK_API const TCHAR* GCancelExport;
+
+// Updates
+extern CAPTUREPROTOCOLSTACK_API const TCHAR* GSessionStopped;
+extern CAPTUREPROTOCOLSTACK_API const TCHAR* GRecordingStatus;
+extern CAPTUREPROTOCOLSTACK_API const TCHAR* GTakeAdded;
+extern CAPTUREPROTOCOLSTACK_API const TCHAR* GTakeRemoved;
+extern CAPTUREPROTOCOLSTACK_API const TCHAR* GTakeUpdated;
+
+// Platform Specific: IOS
+extern CAPTUREPROTOCOLSTACK_API const TCHAR* GDiskCapacity;
+extern CAPTUREPROTOCOLSTACK_API const TCHAR* GBattery;
+extern CAPTUREPROTOCOLSTACK_API const TCHAR* GThermalState;
+}
+
+namespace UE::CPS::Properties
+{
+// Keys
+// Messages
+extern CAPTUREPROTOCOLSTACK_API const TCHAR* GSessionId;
+extern CAPTUREPROTOCOLSTACK_API const TCHAR* GAddressPath;
+extern CAPTUREPROTOCOLSTACK_API const TCHAR* GTransactionId;
+extern CAPTUREPROTOCOLSTACK_API const TCHAR* GTimestamp;
+extern CAPTUREPROTOCOLSTACK_API const TCHAR* GType;
+extern CAPTUREPROTOCOLSTACK_API const TCHAR* GBody;
+extern CAPTUREPROTOCOLSTACK_API const TCHAR* GError;
+extern CAPTUREPROTOCOLSTACK_API const TCHAR* GDescription;
+
+// Requests
+extern CAPTUREPROTOCOLSTACK_API const TCHAR* GSlateName;
+extern CAPTUREPROTOCOLSTACK_API const TCHAR* GTakeNumber;
+extern CAPTUREPROTOCOLSTACK_API const TCHAR* GSubject;
+extern CAPTUREPROTOCOLSTACK_API const TCHAR* GScenario;
+extern CAPTUREPROTOCOLSTACK_API const TCHAR* GTags;
+extern CAPTUREPROTOCOLSTACK_API const TCHAR* GNames;
+extern CAPTUREPROTOCOLSTACK_API const TCHAR* GTakeName;
+extern CAPTUREPROTOCOLSTACK_API const TCHAR* GFiles;
+extern CAPTUREPROTOCOLSTACK_API const TCHAR* GName;
+extern CAPTUREPROTOCOLSTACK_API const TCHAR* GOffset;
+
+
+// Responses
+extern CAPTUREPROTOCOLSTACK_API const TCHAR* GId;
+extern CAPTUREPROTOCOLSTACK_API const TCHAR* GModel;
+extern CAPTUREPROTOCOLSTACK_API const TCHAR* GPlatformName;
+extern CAPTUREPROTOCOLSTACK_API const TCHAR* GPlatformVersion;
+extern CAPTUREPROTOCOLSTACK_API const TCHAR* GSoftwareName;
+extern CAPTUREPROTOCOLSTACK_API const TCHAR* GSoftwareVersion;
+extern CAPTUREPROTOCOLSTACK_API const TCHAR* GExportPort;
+extern CAPTUREPROTOCOLSTACK_API const TCHAR* GIsRecording;
+extern CAPTUREPROTOCOLSTACK_API const TCHAR* GPlatformState;
+extern CAPTUREPROTOCOLSTACK_API const TCHAR* GTakes;
+extern CAPTUREPROTOCOLSTACK_API const TCHAR* GDateTime;
+extern CAPTUREPROTOCOLSTACK_API const TCHAR* GFrames;
+extern CAPTUREPROTOCOLSTACK_API const TCHAR* GAppVersion;
+extern CAPTUREPROTOCOLSTACK_API const TCHAR* GLength;
+extern CAPTUREPROTOCOLSTACK_API const TCHAR* GVideo;
+extern CAPTUREPROTOCOLSTACK_API const TCHAR* GFrameRate;
+extern CAPTUREPROTOCOLSTACK_API const TCHAR* GHeight;
+extern CAPTUREPROTOCOLSTACK_API const TCHAR* GWidth;
+extern CAPTUREPROTOCOLSTACK_API const TCHAR* GAudio;
+extern CAPTUREPROTOCOLSTACK_API const TCHAR* GChannels;
+extern CAPTUREPROTOCOLSTACK_API const TCHAR* GSampleRate;
+extern CAPTUREPROTOCOLSTACK_API const TCHAR* GBitsPerChannel;
+
+// Platform Specific: iOS
+extern CAPTUREPROTOCOLSTACK_API const TCHAR* GTotalCapacity;
+extern CAPTUREPROTOCOLSTACK_API const TCHAR* GRemainingCapacity;
+extern CAPTUREPROTOCOLSTACK_API const TCHAR* GBatteryLevel;
+extern CAPTUREPROTOCOLSTACK_API const TCHAR* GThermalState;
+
+extern CAPTUREPROTOCOLSTACK_API const TCHAR* GTotal;
+extern CAPTUREPROTOCOLSTACK_API const TCHAR* GRemaining;
+extern CAPTUREPROTOCOLSTACK_API const TCHAR* GLevel;
+extern CAPTUREPROTOCOLSTACK_API const TCHAR* GState;
+
+// Values
+// Messages
+extern CAPTUREPROTOCOLSTACK_API const TCHAR* GRequest;
+extern CAPTUREPROTOCOLSTACK_API const TCHAR* GResponse;
+extern CAPTUREPROTOCOLSTACK_API const TCHAR* GUpdate;
+
+extern CAPTUREPROTOCOLSTACK_API const TCHAR* GIOS;
+extern CAPTUREPROTOCOLSTACK_API const TCHAR* GAndroid;
+extern CAPTUREPROTOCOLSTACK_API const TCHAR* GWindows;
+extern CAPTUREPROTOCOLSTACK_API const TCHAR* GWindowsServer;
+extern CAPTUREPROTOCOLSTACK_API const TCHAR* GLinux;
+extern CAPTUREPROTOCOLSTACK_API const TCHAR* GMacOS;
+
+// Platform Specific: iOS
+extern CAPTUREPROTOCOLSTACK_API const TCHAR* GNominal;
+extern CAPTUREPROTOCOLSTACK_API const TCHAR* GFair;
+extern CAPTUREPROTOCOLSTACK_API const TCHAR* GSerious;
+extern CAPTUREPROTOCOLSTACK_API const TCHAR* GCritical;
+}
+
+namespace UE::CPS::ErrorNames
+{
+extern CAPTUREPROTOCOLSTACK_API const TCHAR* GUnsupportedProtocolVersion;
+extern CAPTUREPROTOCOLSTACK_API const TCHAR* GRequestFailed;
+extern CAPTUREPROTOCOLSTACK_API const TCHAR* GRequestParsingFailed;
+extern CAPTUREPROTOCOLSTACK_API const TCHAR* GNotSupported;
+extern CAPTUREPROTOCOLSTACK_API const TCHAR* GNoSession;
+extern CAPTUREPROTOCOLSTACK_API const TCHAR* GMissingBody;
+}
\ No newline at end of file
diff --git a/Engine/Plugins/VirtualProduction/CaptureManager/Source/CaptureProtocolStack/Public/Control/Messages/ControlJsonUtilities.h b/Engine/Plugins/VirtualProduction/CaptureManager/Source/CaptureProtocolStack/Public/Control/Messages/ControlJsonUtilities.h
new file mode 100644
--- /dev/null
+++ b/Engine/Plugins/VirtualProduction/CaptureManager/Source/CaptureProtocolStack/Public/Control/Messages/ControlJsonUtilities.h
@@ -0,0 +1,35 @@
+// Copyright Epic Games, Inc. All Rights Reserved.
+
+#pragma once
+
+#include "Utility/Error.h"
+
+#include "Serialization/JsonReader.h"
+#include "Serialization/JsonWriter.h"
+#include "Serialization/JsonSerializer.h"
+
+#include "Policies/CondensedJsonPrintPolicy.h"
+
+#define CHECK_PARSE CPS_CHECK_VOID_RESULT
+
+class CAPTUREPROTOCOLSTACK_API FJsonUtility final
+{
+public:
+	static bool CreateJsonFromUTF8Data(const TArray<uint8>& InData, TSharedPtr<FJsonObject>& OutObject);
+	static bool CreateUTF8DataFromJson(TSharedPtr<FJsonObject> InObject, TArray<uint8>& OutData);
+
+    template <typename T>
+    static TProtocolResult<void> ParseNumber(const TSharedPtr<FJsonObject>& InBody, const FString& InFieldName, T& OutFieldValue)
+    {
+        if (!InBody->TryGetNumberField(InFieldName, OutFieldValue))
+        {
+            return FCaptureProtocolError(TEXT("Failed to parse key: ") + InFieldName);
+        }
+
+        return ResultOk;
+    }
+    static TProtocolResult<void> ParseString(const TSharedPtr<FJsonObject>& InBody, const FString& InFieldName, FString& OutFieldValue);
+    static TProtocolResult<void> ParseBool(const TSharedPtr<FJsonObject>& InBody, const FString& InFieldName, bool& OutFieldValue);
+    static TProtocolResult<void> ParseObject(const TSharedPtr<FJsonObject>& InBody, const FString& InFieldName, const TSharedPtr<FJsonObject>*& OutFieldValue);
+    static TProtocolResult<void> ParseArray(const TSharedPtr<FJsonObject>& InBody, const FString& InFieldName, const TArray<TSharedPtr<FJsonValue>>*& OutFieldValue);
+};
\ No newline at end of file
diff --git a/Engine/Plugins/VirtualProduction/CaptureManager/Source/CaptureProtocolStack/Public/Control/Messages/ControlMessage.h b/Engine/Plugins/VirtualProduction/CaptureManager/Source/CaptureProtocolStack/Public/Control/Messages/ControlMessage.h
new file mode 100644
--- /dev/null
+++ b/Engine/Plugins/VirtualProduction/CaptureManager/Source/CaptureProtocolStack/Public/Control/Messages/ControlMessage.h
@@ -0,0 +1,72 @@
+// Copyright Epic Games, Inc. All Rights Reserved.
+
+#pragma once
+
+#include "Control/Communication/ControlPacket.h"
+
+class CAPTUREPROTOCOLSTACK_API FControlMessage final
+{
+public:
+
+    enum class EType
+    {
+        Request,
+        Response,
+        Update,
+
+        Invalid
+    };
+
+    FControlMessage(FString InAddressPath, EType InType, TSharedPtr<FJsonObject> InBody);
+
+    FControlMessage(const FControlMessage& InOther) = default;
+    FControlMessage(FControlMessage&& InOther) = default;
+
+    FControlMessage& operator=(const FControlMessage& InOther) = default;
+    FControlMessage& operator=(FControlMessage&& InOther) = default;
+
+    static TProtocolResult<FControlMessage> Deserialize(const FControlPacket& InPacket);
+    static TProtocolResult<FControlPacket> Serialize(const FControlMessage& InMessage);
+
+    void SetSessionId(FString InSessionId);
+    void SetTransactionId(uint32 InTrasactionId);
+    void SetTimestamp(uint64 InTimestamp);
+
+    const FString& GetSessionId() const;
+    const FString& GetAddressPath() const;
+    uint32 GetTransactionId() const;
+    uint64 GetTimestamp() const;
+    EType GetType() const;
+    const TSharedPtr<FJsonObject>& GetBody() const;
+    TSharedPtr<FJsonObject>& GetBody();
+
+    const FString& GetErrorName() const;
+    const FString& GetErrorDescription() const;
+
+private:
+
+    struct FErrorResponse
+    {
+        FString Name = TEXT("");
+        FString Description = TEXT("");
+    };
+
+    FControlMessage(FString InSessionId, 
+                    FString InAddressPath, 
+                    uint32 InTransactionId, 
+                    uint64 InTimestamp, 
+                    EType InMessageType, 
+                    TSharedPtr<FJsonObject> InBody,
+                    FErrorResponse InError);
+
+    static EType DeserializeType(FString InMessageTypeStr);
+    static FString SerializeType(EType InMessageType);
+
+    FString SessionId;
+    FString AddressPath;
+    uint32 TransactionId;
+    uint64 Timestamp;
+    EType MessageType;
+    TSharedPtr<FJsonObject> Body;
+    FErrorResponse Error;
+};
diff --git a/Engine/Plugins/VirtualProduction/CaptureManager/Source/CaptureProtocolStack/Public/Control/Messages/ControlRequest.h b/Engine/Plugins/VirtualProduction/CaptureManager/Source/CaptureProtocolStack/Public/Control/Messages/ControlRequest.h
new file mode 100644
--- /dev/null
+++ b/Engine/Plugins/VirtualProduction/CaptureManager/Source/CaptureProtocolStack/Public/Control/Messages/ControlRequest.h
@@ -0,0 +1,150 @@
+// Copyright Epic Games, Inc. All Rights Reserved.
+
+#pragma once
+
+#include "ControlResponse.h"
+#include "Misc/Optional.h"
+
+class CAPTUREPROTOCOLSTACK_API FControlRequest
+{
+public:
+
+    FControlRequest(FString InAddressPath);
+    virtual ~FControlRequest() = default;
+
+    const FString& GetAddressPath() const;
+    virtual TSharedPtr<FJsonObject> GetBody() const;
+
+private:
+
+    FString AddressPath;
+};
+
+class CAPTUREPROTOCOLSTACK_API FKeepAliveRequest final : public FControlRequest
+{
+public:
+
+    using ResponseType = FKeepAliveResponse;
+
+    FKeepAliveRequest();
+};
+
+class CAPTUREPROTOCOLSTACK_API FStartSessionRequest final : public FControlRequest
+{
+public:
+
+    using ResponseType = FStartSessionResponse;
+
+    FStartSessionRequest();
+};
+
+class CAPTUREPROTOCOLSTACK_API FStopSessionRequest final : public FControlRequest
+{
+public:
+
+    using ResponseType = FStopSessionResponse;
+
+    FStopSessionRequest();
+};
+
+
+class CAPTUREPROTOCOLSTACK_API FGetServerInformationRequest final : public FControlRequest
+{
+public:
+
+    using ResponseType = FGetServerInformationResponse;
+
+    FGetServerInformationRequest();
+};
+
+class CAPTUREPROTOCOLSTACK_API FSubscribeRequest final : public FControlRequest
+{
+public:
+
+    using ResponseType = FSubscribeResponse;
+
+    FSubscribeRequest();
+};
+
+class CAPTUREPROTOCOLSTACK_API FUnsubscribeRequest final : public FControlRequest
+{
+public:
+
+    using ResponseType = FUnsubscribeResponse;
+
+    FUnsubscribeRequest();
+};
+
+class CAPTUREPROTOCOLSTACK_API FGetStateRequest final : public FControlRequest
+{
+public:
+
+    using ResponseType = FGetStateResponse;
+
+    FGetStateRequest();
+};
+
+class CAPTUREPROTOCOLSTACK_API FStartRecordingTakeRequest final : public FControlRequest
+{
+public:
+
+    using ResponseType = FStartRecordingTakeResponse;
+
+    FStartRecordingTakeRequest(FString InSlateName, 
+                               uint16 InTakeNumber, 
+                               TOptional<FString> InSubject = TOptional<FString>(),
+							   TOptional<FString> InScenario = TOptional<FString>(),
+                               TOptional<TArray<FString>> InTags = TOptional<TArray<FString>>());
+
+    virtual TSharedPtr<FJsonObject> GetBody() const override;
+
+private:
+
+    FString SlateName;
+    uint16 TakeNumber;
+	TOptional<FString> Subject;
+	TOptional<FString> Scenario;
+	TOptional<TArray<FString>> Tags;
+};
+
+class CAPTUREPROTOCOLSTACK_API FStopRecordingTakeRequest final : public FControlRequest
+{
+public:
+
+    using ResponseType = FStopRecordingTakeResponse;
+
+    FStopRecordingTakeRequest();
+};
+
+class CAPTUREPROTOCOLSTACK_API FAbortRecordingTakeRequest final : public FControlRequest
+{
+public:
+
+	using ResponseType = FAbortRecordingTakeResponse;
+
+	FAbortRecordingTakeRequest();
+};
+
+class CAPTUREPROTOCOLSTACK_API FGetTakeListRequest final : public FControlRequest
+{
+public:
+
+    using ResponseType = FGetTakeListResponse;
+
+    FGetTakeListRequest();
+};
+
+class CAPTUREPROTOCOLSTACK_API FGetTakeMetadataRequest final : public FControlRequest
+{
+public:
+
+    using ResponseType = FGetTakeMetadataResponse;
+
+    FGetTakeMetadataRequest(TArray<FString> InNames);
+
+    virtual TSharedPtr<FJsonObject> GetBody() const override;
+
+private:
+
+    TArray<FString> Names;
+};
diff --git a/Engine/Plugins/VirtualProduction/CaptureManager/Source/CaptureProtocolStack/Public/Control/Messages/ControlResponse.h b/Engine/Plugins/VirtualProduction/CaptureManager/Source/CaptureProtocolStack/Public/Control/Messages/ControlResponse.h
new file mode 100644
--- /dev/null
+++ b/Engine/Plugins/VirtualProduction/CaptureManager/Source/CaptureProtocolStack/Public/Control/Messages/ControlResponse.h
@@ -0,0 +1,211 @@
+// Copyright Epic Games, Inc. All Rights Reserved.
+
+#pragma once
+
+#include "ControlJsonUtilities.h"
+
+class CAPTUREPROTOCOLSTACK_API FControlResponse
+{
+public:
+
+    FControlResponse(FString InAddressPath);
+    virtual ~FControlResponse() = default;
+
+    virtual TProtocolResult<void> Parse(TSharedPtr<FJsonObject> InBody);
+
+    const FString& GetAddressPath() const;
+
+private:
+
+    FString AddressPath;
+};
+
+class CAPTUREPROTOCOLSTACK_API FKeepAliveResponse final : public FControlResponse
+{
+public:
+
+    FKeepAliveResponse();
+};
+
+class CAPTUREPROTOCOLSTACK_API FStartSessionResponse final : public FControlResponse
+{
+public:
+    
+    FStartSessionResponse();
+
+    virtual TProtocolResult<void> Parse(TSharedPtr<FJsonObject> InBody) override;
+
+    const FString& GetSessionId() const;
+
+private:
+
+    FString SessionId;
+};
+
+class CAPTUREPROTOCOLSTACK_API FStopSessionResponse final : public FControlResponse
+{
+public:
+
+    FStopSessionResponse();
+};
+
+class CAPTUREPROTOCOLSTACK_API FGetServerInformationResponse final : public FControlResponse
+{
+public:
+    
+    FGetServerInformationResponse();
+
+    virtual TProtocolResult<void> Parse(TSharedPtr<FJsonObject> InBody) override;
+
+    const FString& GetId() const;
+    const FString& GetName() const;
+    const FString& GetModel() const;
+    const FString& GetPlatformName() const;
+    const FString& GetPlatformVersion() const;
+    const FString& GetSoftwareName() const;
+    const FString& GetSoftwareVersion() const;
+	uint16 GetExportPort() const;
+
+private:
+
+    FString Id;
+    FString Name;
+    FString Model;
+    FString PlatformName;
+    FString PlatformVersion;
+    FString SoftwareName;
+    FString SoftwareVersion;
+	uint16 ExportPort;
+};
+
+class CAPTUREPROTOCOLSTACK_API FSubscribeResponse final : public FControlResponse
+{
+public:
+
+    FSubscribeResponse();
+};
+
+class CAPTUREPROTOCOLSTACK_API FUnsubscribeResponse final : public FControlResponse
+{
+public:
+
+    FUnsubscribeResponse();
+};
+
+
+class CAPTUREPROTOCOLSTACK_API FGetStateResponse final : public FControlResponse
+{
+public:
+
+    FGetStateResponse();
+
+    virtual TProtocolResult<void> Parse(TSharedPtr<FJsonObject> InBody) override;
+
+    bool IsRecording() const;
+    const TSharedPtr<FJsonObject>& GetPlatformState() const;
+
+private:
+
+    bool bIsRecording;
+    TSharedPtr<FJsonObject> PlatformState;
+};
+
+class CAPTUREPROTOCOLSTACK_API FStartRecordingTakeResponse final : public FControlResponse
+{
+public:
+
+    FStartRecordingTakeResponse();
+};
+
+class CAPTUREPROTOCOLSTACK_API FStopRecordingTakeResponse final : public FControlResponse
+{
+public:
+
+    FStopRecordingTakeResponse();
+
+	virtual TProtocolResult<void> Parse(TSharedPtr<FJsonObject> InBody) override;
+
+	const FString& GetTakeName() const;
+private:
+
+	FString TakeName;
+};
+
+class CAPTUREPROTOCOLSTACK_API FAbortRecordingTakeResponse final : public FControlResponse
+{
+public:
+
+	FAbortRecordingTakeResponse();
+};
+
+class CAPTUREPROTOCOLSTACK_API FGetTakeListResponse final : public FControlResponse
+{
+public:
+
+    FGetTakeListResponse();
+
+    virtual TProtocolResult<void> Parse(TSharedPtr<FJsonObject> InBody) override;
+
+    const TArray<FString>& GetNames() const;
+
+private:
+
+    TArray<FString> Names;
+};
+
+class CAPTUREPROTOCOLSTACK_API FGetTakeMetadataResponse final : public FControlResponse
+{
+public:
+
+    struct FFileObject
+    {
+        FString Name;
+        uint64 Length;
+    };
+
+    struct FVideoObject
+    {
+		uint64 Frames;
+        uint16 FrameRate;
+        uint32 Height;
+        uint32 Width;
+    };
+
+    struct FAudioObject
+    {
+        uint8 Channels;
+        uint32 SampleRate;
+        uint8 BitsPerChannel;
+    };
+
+    struct FTakeObject
+    {
+        FString Name;
+        FString Slate;
+        uint16 TakeNumber;
+        FString DateTime;
+        FString AppVersion;
+        FString Model;
+        FString Subject;
+        FString Scenario;
+        TArray<FString> Tags;
+
+        TArray<FFileObject> Files;
+
+        FVideoObject Video;
+        FAudioObject Audio;
+    };
+
+    FGetTakeMetadataResponse();
+
+    virtual TProtocolResult<void> Parse(TSharedPtr<FJsonObject> InBody) override;
+
+    const TArray<FTakeObject>& GetTakes() const;
+
+private:
+
+    TProtocolResult<void> CreateTakeObject(const TSharedPtr<FJsonObject>& InTakeObject, FTakeObject& OutTake) const;
+    TProtocolResult<void> CreateFileObject(const TSharedPtr<FJsonObject>& InFileObject, FFileObject& OutFile) const;
+
+    TArray<FTakeObject> Takes;
+};
diff --git a/Engine/Plugins/VirtualProduction/CaptureManager/Source/CaptureProtocolStack/Public/Control/Messages/ControlUpdate.h b/Engine/Plugins/VirtualProduction/CaptureManager/Source/CaptureProtocolStack/Public/Control/Messages/ControlUpdate.h
new file mode 100644
--- /dev/null
+++ b/Engine/Plugins/VirtualProduction/CaptureManager/Source/CaptureProtocolStack/Public/Control/Messages/ControlUpdate.h
@@ -0,0 +1,153 @@
+// Copyright Epic Games, Inc. All Rights Reserved.
+
+#pragma once
+
+#include "ControlJsonUtilities.h"
+
+#include "Utility/Error.h"
+
+#include "Delegates/Delegate.h"
+
+class FControlUpdate;
+class FControlUpdateCreator final
+{
+public:
+    static TProtocolResult<TSharedRef<FControlUpdate>> Create(const FString& InAddressPath);
+};
+
+class CAPTUREPROTOCOLSTACK_API FControlUpdate
+{
+public:
+
+    DECLARE_DELEGATE_OneParam(FOnUpdateMessage, TSharedPtr<FControlUpdate> InUpdateMessage);
+
+    FControlUpdate(FString InAddressPath);
+    virtual ~FControlUpdate() = default;
+
+    const FString& GetAddressPath() const;
+
+    virtual TProtocolResult<void> Parse(TSharedPtr<FJsonObject> InBody);
+
+private:
+
+    FString AddressPath;
+};
+
+class CAPTUREPROTOCOLSTACK_API FSessionStopped final : public FControlUpdate
+{
+public:
+
+    FSessionStopped();
+};
+
+class CAPTUREPROTOCOLSTACK_API FRecordingStatusUpdate final : public FControlUpdate
+{
+public:
+
+    FRecordingStatusUpdate();
+
+    virtual TProtocolResult<void> Parse(TSharedPtr<FJsonObject> InBody) override;
+
+    bool IsRecording() const;
+
+private:
+
+    bool bIsRecording;
+};
+
+class CAPTUREPROTOCOLSTACK_API FBaseTakeUpdate : public FControlUpdate
+{
+public:
+
+	using FControlUpdate::FControlUpdate;
+
+	virtual TProtocolResult<void> Parse(TSharedPtr<FJsonObject> InBody) override;
+
+	const FString& GetTakeName() const;
+
+private:
+
+	FString TakeName;
+};
+
+class CAPTUREPROTOCOLSTACK_API FTakeAddedUpdate final : public FBaseTakeUpdate
+{
+public:
+
+	FTakeAddedUpdate();
+};
+
+class CAPTUREPROTOCOLSTACK_API FTakeRemovedUpdate final : public FBaseTakeUpdate
+{
+public:
+
+	FTakeRemovedUpdate();
+};
+
+class CAPTUREPROTOCOLSTACK_API FTakeUpdatedUpdate final : public FBaseTakeUpdate
+{
+public:
+
+	FTakeUpdatedUpdate();
+};
+
+// IOS
+
+class CAPTUREPROTOCOLSTACK_API FDiskCapacityUpdate final : public FControlUpdate
+{
+public:
+
+    FDiskCapacityUpdate();
+
+    virtual TProtocolResult<void> Parse(TSharedPtr<FJsonObject> InBody) override;
+
+    uint64 GetTotal() const;
+    uint64 GetRemaining() const;
+
+private:
+
+    uint64 Total;
+    uint64 Remaining;
+};
+
+class CAPTUREPROTOCOLSTACK_API FBatteryPercentageUpdate final : public FControlUpdate
+{
+public:
+
+    FBatteryPercentageUpdate();
+
+    virtual TProtocolResult<void> Parse(TSharedPtr<FJsonObject> InBody) override;
+
+    float GetLevel() const;
+
+private:
+
+    float Level;
+};
+
+class CAPTUREPROTOCOLSTACK_API FThermalStateUpdate final : public FControlUpdate
+{
+public:
+
+    enum class EState
+    {
+        Nominal = 0,
+        Fair,
+        Serious,
+        Critical,
+
+        Invalid
+    };
+
+    FThermalStateUpdate();
+
+    virtual TProtocolResult<void> Parse(TSharedPtr<FJsonObject> InBody) override;
+
+    EState GetState() const;
+
+private:
+
+    static EState ConvertState(const FString& InStateString);
+
+    EState State;
+};
\ No newline at end of file
diff --git a/Engine/Plugins/VirtualProduction/CaptureManager/Source/CaptureProtocolStack/Public/Discovery/Communication/DiscoveryCommunication.h b/Engine/Plugins/VirtualProduction/CaptureManager/Source/CaptureProtocolStack/Public/Discovery/Communication/DiscoveryCommunication.h
new file mode 100644
--- /dev/null
+++ b/Engine/Plugins/VirtualProduction/CaptureManager/Source/CaptureProtocolStack/Public/Discovery/Communication/DiscoveryCommunication.h
@@ -0,0 +1,43 @@
+// Copyright Epic Games, Inc. All Rights Reserved.
+
+#pragma once
+
+#include "DiscoveryPacket.h"
+
+#include "Communication/UdpClient.h"
+
+#include "Utility/Error.h"
+#include "Utility/QueueRunner.h"
+
+class FDiscoveryCommunication final
+{
+public:
+
+	static const uint16 MulticastPort;
+    static const FString MulticastAddress;
+
+    DECLARE_DELEGATE_OneParam(FOnPacketReceived, FDiscoveryPacket InPacket)
+
+    using FRunnerType = TQueueRunner<FDiscoveryPacket>;
+
+    FDiscoveryCommunication();
+
+    TProtocolResult<void> Start();
+    TProtocolResult<void> Stop();
+
+    void SendMessage(FDiscoveryPacket InMessage, const FString& Endpoint);
+    void SendMessage(FDiscoveryPacket InMessage, const FString& EndpointIp, const uint16 EndpointPort);
+
+    void SetReceiveHandler(FOnPacketReceived InReceiveHandler);
+
+private:
+
+    void OnProcessReceivedPacket(FDiscoveryPacket InPacket);
+    void OnPacketArrived(const FArrayReaderPtr& InPayload, const FIPv4Endpoint& InEndpoint);
+
+    FUdpClient Client;
+
+    FRunnerType SynchronizedReceiver;
+
+    FOnPacketReceived OnPacketReceived;
+};
\ No newline at end of file
diff --git a/Engine/Plugins/VirtualProduction/CaptureManager/Source/CaptureProtocolStack/Public/Discovery/Communication/DiscoveryPacket.h b/Engine/Plugins/VirtualProduction/CaptureManager/Source/CaptureProtocolStack/Public/Discovery/Communication/DiscoveryPacket.h
new file mode 100644
--- /dev/null
+++ b/Engine/Plugins/VirtualProduction/CaptureManager/Source/CaptureProtocolStack/Public/Discovery/Communication/DiscoveryPacket.h
@@ -0,0 +1,47 @@
+// Copyright Epic Games, Inc. All Rights Reserved.
+
+#pragma once
+
+#include "Utility/Definitions.h"
+#include "Utility/Error.h"
+
+#include "Containers/Array.h"
+
+class FDiscoveryPacket final
+{
+public:
+
+	enum class EMessageType : uint8
+	{
+		Request = 0,
+		Response = 1,
+		Notify = 2,
+		
+		Invalid = 255
+	};
+
+    static const TArray<uint8> Header;
+
+    static TProtocolResult<FDiscoveryPacket> Deserialize(const TConstArrayView<uint8>& InData);
+    static TProtocolResult<TArray<uint8>> Serialize(const FDiscoveryPacket& InPacket);
+    
+    FDiscoveryPacket() = default;
+    FDiscoveryPacket(EMessageType InMessageType, TArray<uint8> InPayload);
+
+    FDiscoveryPacket(const FDiscoveryPacket& InOther) = default;
+    FDiscoveryPacket(FDiscoveryPacket&& InOther) = default;
+
+    FDiscoveryPacket& operator=(const FDiscoveryPacket& InOther) = default;
+    FDiscoveryPacket& operator=(FDiscoveryPacket&& InOther) = default;
+
+	const EMessageType GetMessageType() const;
+    const TArray<uint8>& GetPayload() const;
+
+private:
+
+	static EMessageType ToMessageType(uint8 InMessageType);
+	static uint8 FromMessageType(EMessageType InMessageType);
+    
+	EMessageType MessageType;
+    TArray<uint8> Payload;
+};
\ No newline at end of file
diff --git a/Engine/Plugins/VirtualProduction/CaptureManager/Source/CaptureProtocolStack/Public/Discovery/DiscoveryMessenger.h b/Engine/Plugins/VirtualProduction/CaptureManager/Source/CaptureProtocolStack/Public/Discovery/DiscoveryMessenger.h
new file mode 100644
--- /dev/null
+++ b/Engine/Plugins/VirtualProduction/CaptureManager/Source/CaptureProtocolStack/Public/Discovery/DiscoveryMessenger.h
@@ -0,0 +1,40 @@
+// Copyright Epic Games, Inc. All Rights Reserved.
+
+#pragma once
+
+#include "Discovery/Communication/DiscoveryCommunication.h"
+#include "Discovery/Communication/DiscoveryPacket.h"
+
+#include "Discovery/Messages/DiscoveryRequest.h"
+#include "Discovery/Messages/DiscoveryResponse.h"
+#include "Discovery/Messages/DiscoveryNotify.h"
+
+#include "Utility/Error.h"
+
+class CAPTUREPROTOCOLSTACK_API FDiscoveryMessenger final
+{
+public:
+
+    DECLARE_DELEGATE_OneParam(FOnResponseArrived, FDiscoveryResponse InResponse);
+    DECLARE_DELEGATE_OneParam(FOnNotifyArrived, FDiscoveryNotify InResponse);
+
+    FDiscoveryMessenger();
+    ~FDiscoveryMessenger();
+
+    TProtocolResult<void> Start();
+    TProtocolResult<void> Stop();
+
+    TProtocolResult<void> SendMulticastRequest();
+
+    void SetResponseHandler(FOnResponseArrived InOnResponse);
+    void SetNotifyHandler(FOnNotifyArrived InOnNotify);
+
+private:
+
+    void OnPacketArrived(FDiscoveryPacket InPacket);
+
+    FDiscoveryCommunication Communication;
+
+    FOnResponseArrived OnResponse;
+    FOnNotifyArrived OnNotify;
+};
diff --git a/Engine/Plugins/VirtualProduction/CaptureManager/Source/CaptureProtocolStack/Public/Discovery/Messages/DiscoveryNotify.h b/Engine/Plugins/VirtualProduction/CaptureManager/Source/CaptureProtocolStack/Public/Discovery/Messages/DiscoveryNotify.h
new file mode 100644
--- /dev/null
+++ b/Engine/Plugins/VirtualProduction/CaptureManager/Source/CaptureProtocolStack/Public/Discovery/Messages/DiscoveryNotify.h
@@ -0,0 +1,47 @@
+// Copyright Epic Games, Inc. All Rights Reserved.
+
+#pragma once
+
+#include "Discovery/Communication/DiscoveryPacket.h"
+
+#include "Utility/Error.h"
+
+#include "Containers/Array.h"
+#include "Containers/StaticArray.h"
+
+class CAPTUREPROTOCOLSTACK_API FDiscoveryNotify
+{
+public:
+
+	static const uint32 MinPayloadSize;
+
+    using FServerId = TStaticArray<uint8, 16>;
+
+    enum class EConnectionState : uint8
+    {
+        Offline = 0,
+        Online = 1,
+
+        Invalid
+    };
+
+    FDiscoveryNotify(FServerId InServerId, uint16 InControlPort, EConnectionState InConnectionState, TArray<uint16> InSupportedVersions);
+
+    static TProtocolResult<FDiscoveryNotify> Deserialize(const FDiscoveryPacket& InPacket);
+    static TProtocolResult<FDiscoveryPacket> Serialize(const FDiscoveryNotify& InNotify);
+
+    const FServerId& GetServerId() const;
+    uint16 GetControlPort() const;
+    EConnectionState GetConnectionState() const;
+	const TArray<uint16>& GetSupportedVersions() const;
+
+private:
+
+    static EConnectionState ToConnectionState(uint8 InConnectionState);
+    static uint8 FromConnectionState(EConnectionState InConnectionState);
+
+    FServerId ServerId;
+    uint16 ControlPort;
+    EConnectionState ConnectionState;
+	TArray<uint16> SupportedVersions;
+};
\ No newline at end of file
diff --git a/Engine/Plugins/VirtualProduction/CaptureManager/Source/CaptureProtocolStack/Public/Discovery/Messages/DiscoveryRequest.h b/Engine/Plugins/VirtualProduction/CaptureManager/Source/CaptureProtocolStack/Public/Discovery/Messages/DiscoveryRequest.h
new file mode 100644
--- /dev/null
+++ b/Engine/Plugins/VirtualProduction/CaptureManager/Source/CaptureProtocolStack/Public/Discovery/Messages/DiscoveryRequest.h
@@ -0,0 +1,22 @@
+// Copyright Epic Games, Inc. All Rights Reserved.
+
+#pragma once
+
+#include "Discovery/Communication/DiscoveryPacket.h"
+
+#include "Utility/Error.h"
+
+#include "Containers/Array.h"
+#include "Containers/StaticArray.h"
+
+class FDiscoveryRequest
+{
+public:
+
+	static const uint32 PayloadSize;
+
+    FDiscoveryRequest() = default;
+
+    static TProtocolResult<FDiscoveryRequest> Deserialize(const FDiscoveryPacket& InPacket);
+    static TProtocolResult<FDiscoveryPacket> Serialize(const FDiscoveryRequest& InRequest);
+};
\ No newline at end of file
diff --git a/Engine/Plugins/VirtualProduction/CaptureManager/Source/CaptureProtocolStack/Public/Discovery/Messages/DiscoveryResponse.h b/Engine/Plugins/VirtualProduction/CaptureManager/Source/CaptureProtocolStack/Public/Discovery/Messages/DiscoveryResponse.h
new file mode 100644
--- /dev/null
+++ b/Engine/Plugins/VirtualProduction/CaptureManager/Source/CaptureProtocolStack/Public/Discovery/Messages/DiscoveryResponse.h
@@ -0,0 +1,34 @@
+// Copyright Epic Games, Inc. All Rights Reserved.
+
+#pragma once
+
+#include "Discovery/Communication/DiscoveryPacket.h"
+
+#include "Utility/Error.h"
+
+#include "Containers/Array.h"
+#include "Containers/StaticArray.h"
+
+class CAPTUREPROTOCOLSTACK_API FDiscoveryResponse
+{
+public:
+
+	static const uint32 MinPayloadSize;
+
+    using FServerId = TStaticArray<uint8, 16>;
+
+    FDiscoveryResponse(FServerId InServerId, uint16 InControlPort, TArray<uint16> InSupportedVersions);
+
+    static TProtocolResult<FDiscoveryResponse> Deserialize(const FDiscoveryPacket& InPacket);
+    static TProtocolResult<FDiscoveryPacket> Serialize(const FDiscoveryResponse& InResponse);
+
+    const FServerId& GetServerId() const;
+    uint16 GetControlPort() const;
+	const TArray<uint16>& GetSupportedVersions() const;
+
+private:
+
+    FServerId ServerId;
+    uint16 ControlPort;
+	TArray<uint16> SupportedVersions;
+};
\ No newline at end of file
diff --git a/Engine/Plugins/VirtualProduction/CaptureManager/Source/CaptureProtocolStack/Public/ExportClient/Communication/ExportCommunication.h b/Engine/Plugins/VirtualProduction/CaptureManager/Source/CaptureProtocolStack/Public/ExportClient/Communication/ExportCommunication.h
new file mode 100644
--- /dev/null
+++ b/Engine/Plugins/VirtualProduction/CaptureManager/Source/CaptureProtocolStack/Public/ExportClient/Communication/ExportCommunication.h
@@ -0,0 +1,53 @@
+// Copyright Epic Games, Inc. All Rights Reserved.
+
+#pragma once
+
+#include "ExportHeader.h"
+
+#include "ExportClient/Messages/ExportRequest.h"
+#include "ExportClient/Messages/ExportResponse.h"
+
+#include "Communication/TcpClient.h"
+
+#include "Utility/QueueRunner.h"
+
+#include "HAL/Runnable.h"
+#include "HAL/RunnableThread.h"
+
+struct FExportRequestPacket
+{
+	static TProtocolResult<void> Serialize(const FExportRequestPacket& InRequestPacket, ITcpSocketWriter& InWriter);
+
+	FExportHeader Header;
+	FExportRequest Request;
+};
+
+struct FExportResponseHeader
+{
+	static TProtocolResult<FExportResponseHeader> Deserialize(ITcpSocketReader& InReader);
+
+	FExportHeader Header;
+	FExportResponse Response;
+};
+
+class FExportCommunication final
+{
+public:
+
+	FExportCommunication();
+
+	TProtocolResult<void> Init();
+	TProtocolResult<void> Start(const FString& InServerIp, const uint16 InServerPort);
+	TProtocolResult<void> Stop();
+
+	bool IsRunning() const;
+
+	TProtocolResult<void> SendRequest(FExportHeader InHeader, FExportRequest InRequest);
+	TProtocolResult<FExportResponseHeader> ReceiveResponseHeader();
+	TProtocolResult<TArray<uint8>> ReceiveResponseData(const uint64 InSize);
+	TProtocolResult<TStaticArray<uint8, 16>> ReceiveFileHash();
+
+private:
+
+	FTcpClient Client;
+};
diff --git a/Engine/Plugins/VirtualProduction/CaptureManager/Source/CaptureProtocolStack/Public/ExportClient/Communication/ExportHeader.h b/Engine/Plugins/VirtualProduction/CaptureManager/Source/CaptureProtocolStack/Public/ExportClient/Communication/ExportHeader.h
new file mode 100644
--- /dev/null
+++ b/Engine/Plugins/VirtualProduction/CaptureManager/Source/CaptureProtocolStack/Public/ExportClient/Communication/ExportHeader.h
@@ -0,0 +1,37 @@
+// Copyright Epic Games, Inc. All Rights Reserved.
+
+#pragma once
+
+#include "Communication/TcpReaderWriter.h"
+
+#include "Utility/Error.h"
+
+#include "Containers/Array.h"
+#include "Containers/StaticArray.h"
+
+class FExportHeader final
+{
+public:
+
+	static const TArray<uint8> Header;
+
+	FExportHeader() = default;
+	FExportHeader(uint16 InVersion, uint32 InTransactionId);
+
+	FExportHeader(const FExportHeader& InOther) = default;
+	FExportHeader(FExportHeader&& InOther) = default;
+
+	FExportHeader& operator=(const FExportHeader& InOther) = default;
+	FExportHeader& operator=(FExportHeader&& InOther) = default;
+
+	static TProtocolResult<FExportHeader> Deserialize(ITcpSocketReader& InReader);
+	static TProtocolResult<void> Serialize(const FExportHeader& InHeader, ITcpSocketWriter& InWriter);
+
+	uint16 GetVersion() const;
+	uint32 GetTransactionId() const;
+
+private:
+
+	uint16 Version;
+	uint32 TransactionId;
+};
\ No newline at end of file
diff --git a/Engine/Plugins/VirtualProduction/CaptureManager/Source/CaptureProtocolStack/Public/ExportClient/ExportClient.h b/Engine/Plugins/VirtualProduction/CaptureManager/Source/CaptureProtocolStack/Public/ExportClient/ExportClient.h
new file mode 100644
--- /dev/null
+++ b/Engine/Plugins/VirtualProduction/CaptureManager/Source/CaptureProtocolStack/Public/ExportClient/ExportClient.h
@@ -0,0 +1,68 @@
+// Copyright Epic Games, Inc. All Rights Reserved.
+
+#pragma once
+
+#include "ExportClient/Communication/ExportCommunication.h"
+#include "ExportWorker.h"
+
+#include "Containers/Map.h"
+
+class FExportTaskStopToken
+{
+public:
+
+	FExportTaskStopToken();
+
+	void Cancel();
+	bool IsCanceled();
+
+private:
+
+	std::atomic_bool bCanceled;
+};
+
+class CAPTUREPROTOCOLSTACK_API FExportClient final
+	: public FExportTaskExecutor
+{
+public:
+
+	using FTaskId = uint32;
+
+	using FTakeFileArray = TArray<FTakeFile>;
+
+	FExportClient(FString InServerIp, const uint16 InExportPort);
+	~FExportClient();
+
+	FTaskId ExportTakeFiles(FString InTakeName, FTakeFileArray InTakeFileArray, TUniquePtr<FBaseStream> InStream);
+	FTaskId ExportFiles(TMap<FString, FTakeFileArray> InTakesFilesMap, TUniquePtr<FBaseStream> InStream);
+
+	void AbortExport(const FTaskId InTaskId);
+	void AbortAllExports();
+
+private:
+
+	FTaskId StartExport(FExportTakeTask::FExportContexts InContexts, TUniquePtr<FBaseStream> InStream);
+
+	virtual void OnTask(TUniquePtr<FExportTakeTask> InTask) override;
+	bool OnExportTask(TUniquePtr<FExportTakeTask> InTask);
+
+	TProtocolResult<TMap<uint32, TPair<FString, FTakeFile>>> SendRequests(FExportTakeTask::FExportContexts InContexts);
+	TProtocolResult<void> ReceiveResponses(TMap<uint32, TPair<FString, FTakeFile>> InResponseMap,
+										   FBaseStream& InBaseStream);
+
+	TProtocolResult<void> Start();
+	TProtocolResult<void> Stop();
+
+	FString ServerIp;
+	uint16 ServerPort;
+	TUniquePtr<FExportCommunication> Communication;
+
+	FExportWorker ExportWorker;
+	TUniquePtr<FRunnableThread> Thread;
+
+	std::atomic<FTaskId> CurrentTaskId;
+
+	uint32 TransactionIdCounter;
+
+	TUniquePtr<FExportTaskStopToken> CurrentTaskStopToken;
+};
\ No newline at end of file
diff --git a/Engine/Plugins/VirtualProduction/CaptureManager/Source/CaptureProtocolStack/Public/ExportClient/ExportWorker.h b/Engine/Plugins/VirtualProduction/CaptureManager/Source/CaptureProtocolStack/Public/ExportClient/ExportWorker.h
new file mode 100644
--- /dev/null
+++ b/Engine/Plugins/VirtualProduction/CaptureManager/Source/CaptureProtocolStack/Public/ExportClient/ExportWorker.h
@@ -0,0 +1,108 @@
+// Copyright Epic Games, Inc. All Rights Reserved.
+
+#pragma once
+
+#include "ExportClient/Communication/ExportCommunication.h"
+
+#include "HAL/UESemaphore.h"
+
+class FBaseStream;
+struct FTakeFile
+{
+	FString FileName;
+	uint64 Length;
+	uint64 Offset;
+};
+
+class FExportTakeTask
+{
+public:
+
+	struct FExportContext
+	{
+		FString TakeName;
+		FTakeFile File;
+	};
+
+	using FExportContexts = TArray<FExportContext>;
+
+	FExportTakeTask(FExportContexts InExportContexts,
+					TUniquePtr<FBaseStream> InStream);
+
+	FExportContexts ExportContexts;
+	TUniquePtr<FBaseStream> Stream;
+};
+
+class FExportTaskExecutor
+{
+public:
+	virtual ~FExportTaskExecutor() = default;
+
+	virtual void OnTask(TUniquePtr<FExportTakeTask> InTask) = 0;
+};
+
+class FExportQueue
+{
+public:
+	static const int32 MaxNumberOfElements = MAX_int32;
+
+	FExportQueue();
+	~FExportQueue() = default;
+
+	TProtocolResult<void> Add(const uint32 InTaskId, TUniquePtr<FExportTakeTask> InElement);
+	TUniquePtr<FExportTakeTask> Pop();
+
+	TProtocolResult<TUniquePtr<FExportTakeTask>> Remove(uint32 InTaskId);
+
+	bool IsEmpty() const;
+
+	TArray<TUniquePtr<FExportTakeTask>> GetAndEmpty();
+
+private:
+
+	mutable FCriticalSection Mutex;
+	FSemaphore Semaphore;
+	TArray<TPair<uint32, TUniquePtr<FExportTakeTask>>> Queue;
+};
+
+class FExportWorker : public FRunnable
+{
+public:
+
+	FExportWorker(FExportTaskExecutor& InClient);
+	virtual ~FExportWorker() override;
+
+	TProtocolResult<void> Add(const uint32 InTaskId, TUniquePtr<FExportTakeTask> InElement);
+	TProtocolResult<TUniquePtr<FExportTakeTask>> Remove(uint32 InTaskId);
+
+	bool IsEmpty() const;
+
+	TArray<TUniquePtr<FExportTakeTask>> GetAndEmpty();
+
+protected:
+
+	virtual uint32 Run() override;
+	virtual void Stop() override;
+
+private:
+
+	std::atomic<bool> bRunning;
+
+	FExportQueue Queue;
+
+	FExportTaskExecutor& Client;
+};
+
+class CAPTUREPROTOCOLSTACK_API FBaseStream
+{
+public:
+
+	FBaseStream() = default;
+	virtual ~FBaseStream() = default;
+
+	virtual bool StartFile(const FString& InTakeName, const FString& InFileName) = 0;
+	virtual bool ProcessData(const FString& InTakeName, const FString& InFileName, const TConstArrayView<uint8>& InData) = 0;
+	virtual bool FinishFile(const FString& InTakeName, const FString& InFileName, const TStaticArray<uint8, 16>& InHash) = 0;
+
+	virtual void Done(TProtocolResult<void> InResult) = 0;
+};
\ No newline at end of file
diff --git a/Engine/Plugins/VirtualProduction/CaptureManager/Source/CaptureProtocolStack/Public/ExportClient/Messages/ExportRequest.h b/Engine/Plugins/VirtualProduction/CaptureManager/Source/CaptureProtocolStack/Public/ExportClient/Messages/ExportRequest.h
new file mode 100644
--- /dev/null
+++ b/Engine/Plugins/VirtualProduction/CaptureManager/Source/CaptureProtocolStack/Public/ExportClient/Messages/ExportRequest.h
@@ -0,0 +1,32 @@
+// Copyright Epic Games, Inc. All Rights Reserved.
+
+#pragma once
+
+#include "ExportClient/Communication/ExportHeader.h"
+
+class FExportRequest final
+{
+public:
+
+	FExportRequest() = default;
+	FExportRequest(FString InTakeName, FString InFileName, uint64 InOffset);
+	
+	FExportRequest(const FExportRequest& InOther) = default;
+	FExportRequest(FExportRequest&& InOther) = default;
+
+	FExportRequest& operator=(const FExportRequest& InOther) = default;
+	FExportRequest& operator=(FExportRequest&& InOther) = default;
+
+	static TProtocolResult<FExportRequest> Deserialize(ITcpSocketReader& InReader);
+	static TProtocolResult<void> Serialize(const FExportRequest& InRequest, ITcpSocketWriter& InWriter);
+
+	const FString& GetTakeName() const;
+	const FString& GetFileName() const;
+	uint64 GetOffset() const;
+
+private:
+
+	FString TakeName;
+	FString FileName;
+	uint64 Offset;
+};
\ No newline at end of file
diff --git a/Engine/Plugins/VirtualProduction/CaptureManager/Source/CaptureProtocolStack/Public/ExportClient/Messages/ExportResponse.h b/Engine/Plugins/VirtualProduction/CaptureManager/Source/CaptureProtocolStack/Public/ExportClient/Messages/ExportResponse.h
new file mode 100644
--- /dev/null
+++ b/Engine/Plugins/VirtualProduction/CaptureManager/Source/CaptureProtocolStack/Public/ExportClient/Messages/ExportResponse.h
@@ -0,0 +1,49 @@
+// Copyright Epic Games, Inc. All Rights Reserved.
+
+#pragma once
+
+#include "ExportClient/Communication/ExportHeader.h"
+
+class FExportResponse final
+{
+public:
+
+	enum class EStatus : uint8
+	{
+		Success = 0,
+		InvalidTakeName = 1,
+		InvalidFileName = 2,
+		InvalidOffset = 3,
+		ServerError = 4,
+		UnsupportedProtocolVersion = 5,
+
+		Reserved
+	};
+
+	FExportResponse(EStatus InStatus, uint64 InLength);
+
+	FExportResponse(const FExportResponse& InOther) = default;
+	FExportResponse(FExportResponse&& InOther) = default;
+
+	FExportResponse& operator=(const FExportResponse& InOther) = default;
+	FExportResponse& operator=(FExportResponse&& InOther) = default;
+
+	static TProtocolResult<FExportResponse> Deserialize(ITcpSocketReader& InReader);
+	static TProtocolResult<TStaticArray<uint8, 16>> DeserializeHash(ITcpSocketReader& InReader);
+
+	static TProtocolResult<void> Serialize(const FExportResponse& InResponse, ITcpSocketWriter& InWriter);
+	static TProtocolResult<void> SerializeHash(const TStaticArray<uint8, 16>& InHash, ITcpSocketWriter& InWriter);
+
+	EStatus GetStatus() const;
+	uint64 GetLength() const;
+
+private:
+
+	FExportResponse() = default;
+
+	static EStatus ToStatus(uint8 InStatus);
+	static uint8 FromStatus(EStatus InStatus);
+
+	EStatus Status;
+	uint64 Length;
+};
\ No newline at end of file
diff --git a/Engine/Plugins/VirtualProduction/CaptureManager/Source/CaptureProtocolStack/Public/Utility/CPSTimerManager.h b/Engine/Plugins/VirtualProduction/CaptureManager/Source/CaptureProtocolStack/Public/Utility/CPSTimerManager.h
new file mode 100644
--- /dev/null
+++ b/Engine/Plugins/VirtualProduction/CaptureManager/Source/CaptureProtocolStack/Public/Utility/CPSTimerManager.h
@@ -0,0 +1,46 @@
+// Copyright Epic Games, Inc. All Rights Reserved.
+
+#pragma once
+
+#include "HAL/RunnableThread.h"
+#include "HAL/Runnable.h"
+#include "Containers/Ticker.h"
+
+#include "TimerManager.h"
+
+class CAPTUREPROTOCOLSTACK_API FCPSTimerManager
+{
+public:
+
+	using FTimerHandle = FTSTicker::FDelegateHandle;
+
+	static const float IdealTimeResolution;
+
+	FCPSTimerManager(const float Resolution = IdealTimeResolution);
+	~FCPSTimerManager();
+
+	FTimerHandle AddTimer(FTimerDelegate InDelegate, float InRate, bool InbLoop = false, float InFirstDelay = 0.f);
+	void RemoveTimer(FTimerHandle Handle);
+
+private:
+
+	class FTimerManagerRunnable final : public FRunnable
+	{
+	public:
+
+		FTimerManagerRunnable(FCPSTimerManager* InOwner, const float InResolution);
+
+		virtual uint32 Run() override;
+		virtual void Stop() override;
+
+	private:
+
+		FCPSTimerManager* Owner;
+		std::atomic_bool bShouldRun = true;
+		const float IdealTimeResolution;
+	};
+
+	FTSTicker Ticker;
+	TUniquePtr<FRunnableThread> Thread;
+	TUniquePtr<FTimerManagerRunnable> Runnable;
+};
\ No newline at end of file
diff --git a/Engine/Plugins/VirtualProduction/CaptureManager/Source/CaptureProtocolStack/Public/Utility/Definitions.h b/Engine/Plugins/VirtualProduction/CaptureManager/Source/CaptureProtocolStack/Public/Utility/Definitions.h
new file mode 100644
--- /dev/null
+++ b/Engine/Plugins/VirtualProduction/CaptureManager/Source/CaptureProtocolStack/Public/Utility/Definitions.h
@@ -0,0 +1,13 @@
+// Copyright Epic Games, Inc. All Rights Reserved.
+
+#pragma once
+
+#include "Logging/LogMacros.h"
+
+#define CPS_VERSION 1
+
+DECLARE_LOG_CATEGORY_EXTERN(LogCPSControlMessenger, Log, All);
+DECLARE_LOG_CATEGORY_EXTERN(LogCPSControlCommunication, Log, All);
+DECLARE_LOG_CATEGORY_EXTERN(LogCPSDiscoveryMessenger, Log, All);
+DECLARE_LOG_CATEGORY_EXTERN(LogCPSDiscoveryCommunication, Log, All);
+DECLARE_LOG_CATEGORY_EXTERN(LogCPSExportHandler, Log, All);
\ No newline at end of file
diff --git a/Engine/Plugins/VirtualProduction/CaptureManager/Source/CaptureProtocolStack/Public/Utility/Error.h b/Engine/Plugins/VirtualProduction/CaptureManager/Source/CaptureProtocolStack/Public/Utility/Error.h
new file mode 100644
--- /dev/null
+++ b/Engine/Plugins/VirtualProduction/CaptureManager/Source/CaptureProtocolStack/Public/Utility/Error.h
@@ -0,0 +1,106 @@
+// Copyright Epic Games, Inc. All Rights Reserved.
+
+#pragma once
+
+#include "Containers/UnrealString.h"
+
+#include "Templates/ValueOrError.h"
+
+class CAPTUREPROTOCOLSTACK_API FCaptureProtocolError
+{
+public:
+
+	FCaptureProtocolError();
+	FCaptureProtocolError(FString InMessage, int32 InCode = 0);
+
+	FCaptureProtocolError(const FCaptureProtocolError& InOther) = default;
+	FCaptureProtocolError(FCaptureProtocolError&& InOther) = default;
+
+	FCaptureProtocolError& operator=(const FCaptureProtocolError& InOther) = default;
+	FCaptureProtocolError& operator=(FCaptureProtocolError&& InOther) = default;
+
+	const FString& GetMessage() const;
+	int32 GetCode() const;
+
+private:
+
+	FString Message;
+	int32 Code;
+};
+
+template <typename ValueType>
+class TProtocolResult : public TValueOrError<ValueType, FCaptureProtocolError>
+{
+private:
+
+	using Base = TValueOrError<ValueType, FCaptureProtocolError>;
+
+public:
+
+	TProtocolResult(ValueType&& InValueType)
+		: Base(MakeValue(MoveTemp(InValueType)))
+	{
+	}
+
+	TProtocolResult(const ValueType& InValueType)
+		: Base(MakeValue(InValueType))
+	{
+	}
+
+	TProtocolResult(FCaptureProtocolError&& InErrorType)
+		: Base(MakeError(MoveTemp(InErrorType)))
+	{
+	}
+
+	TProtocolResult(const FCaptureProtocolError& InErrorType)
+		: Base(MakeError(InErrorType))
+	{
+	}
+
+	template <typename... ArgTypes>
+	TProtocolResult(TInPlaceType<ValueType>&&, ArgTypes&&... Args)
+		: Base(MakeValue(Forward<ArgTypes>(Args)...))
+	{
+	}
+
+	template <typename... ArgTypes>
+	TProtocolResult(TInPlaceType<FCaptureProtocolError>&&, ArgTypes&&... Args)
+		: Base(MakeError(Forward<ArgTypes>(Args)...))
+	{
+	}
+};
+
+template <>
+class TProtocolResult<void> : public TValueOrError<void, FCaptureProtocolError>
+{
+private:
+
+	using Base = TValueOrError<void, FCaptureProtocolError>;
+
+public:
+
+	TProtocolResult(TInPlaceType<void>)
+		: Base(MakeValue())
+	{
+	}
+
+	TProtocolResult(FCaptureProtocolError&& InErrorType)
+		: Base(MakeError(MoveTemp(InErrorType)))
+	{
+	}
+
+	TProtocolResult(const FCaptureProtocolError& InErrorType)
+		: Base(MakeError(InErrorType))
+	{
+	}
+
+	template <typename... ArgTypes>
+	TProtocolResult(ArgTypes&&... Args)
+		: Base(MakeError(Forward<ArgTypes>(Args)...))
+	{
+	}
+};
+
+const TProtocolResult<void> ResultOk = TInPlaceType<void>{};
+
+#define CPS_CHECK_VOID_RESULT(Function) if (TProtocolResult<void> Result = Function; Result.HasError()) { return Result.StealError(); }
\ No newline at end of file
diff --git a/Engine/Plugins/VirtualProduction/CaptureManager/Source/CaptureProtocolStack/Public/Utility/QueueRunner.h b/Engine/Plugins/VirtualProduction/CaptureManager/Source/CaptureProtocolStack/Public/Utility/QueueRunner.h
new file mode 100644
--- /dev/null
+++ b/Engine/Plugins/VirtualProduction/CaptureManager/Source/CaptureProtocolStack/Public/Utility/QueueRunner.h
@@ -0,0 +1,84 @@
+// Copyright Epic Games, Inc. All Rights Reserved.
+
+#pragma once
+
+#include "Containers/Queue.h"
+
+#include "HAL/RunnableThread.h"
+#include "HAL/Runnable.h"
+#include "Delegates/Delegate.h"
+
+template<class T>
+class TQueueRunner : public FRunnable
+{
+public:
+
+    DECLARE_DELEGATE_OneParam(FOnProcess, T InElem)
+
+    TQueueRunner(FOnProcess InProcess)
+        : bRunning(true)
+        , OnProcess(MoveTemp(InProcess))
+    {
+        Event = FGenericPlatformProcess::GetSynchEventFromPool(false);
+        check(Event);
+
+        Thread.Reset(FRunnableThread::Create(this, TEXT("Queue Runner"), 128 * 1024, TPri_Normal, FPlatformAffinity::GetPoolThreadMask()));
+    }
+
+    ~TQueueRunner()
+    {
+        Thread->Kill(true);
+
+		FGenericPlatformProcess::ReturnSynchEventToPool(Event);
+    }
+
+    void Add(T InElement)
+    {
+		{
+			FScopeLock lock(&Mutex);
+			Queue.Enqueue(MoveTemp(InElement));
+		}
+
+		Event->Trigger();
+    }
+
+protected:
+
+    virtual uint32 Run() override
+    {
+        while (bRunning.load())
+        {
+            Event->Wait();
+
+            FScopeLock Lock(&Mutex);
+            while (!Queue.IsEmpty())
+            {
+                T Elem;
+                Queue.Dequeue(Elem);
+
+                FScopeUnlock Unlock(&Mutex);
+
+                OnProcess.ExecuteIfBound(MoveTemp(Elem));
+            }
+        }
+
+		return 0;
+    }
+
+    virtual void Stop() override
+    {
+        bRunning.store(false);
+
+        Event->Trigger();
+    }
+
+private:
+
+    std::atomic<bool> bRunning;
+    FOnProcess OnProcess;
+
+    FCriticalSection Mutex;
+    FEvent* Event;
+    TQueue<T> Queue;
+    TUniquePtr<FRunnableThread> Thread;
+};
