CIS Build Error: NonUnity: ControlCommunication.cpp EventSourceUtils.cpp

diff --git a/Engine/Plugins/VirtualProduction/CaptureManager/Source/CaptureProtocolStack/Public/Utility/QueueRunner.h b/Engine/Plugins/VirtualProduction/CaptureManager/Source/CaptureProtocolStack/Public/Utility/QueueRunner.h
--- a/Engine/Plugins/VirtualProduction/CaptureManager/Source/CaptureProtocolStack/Public/Utility/QueueRunner.h
+++ b/Engine/Plugins/VirtualProduction/CaptureManager/Source/CaptureProtocolStack/Public/Utility/QueueRunner.h
@@ -4,6 +4,7 @@
 
 #include "Containers/Queue.h"
 
+#include "HAL/Event.h"
 #include "HAL/RunnableThread.h"
 #include "HAL/Runnable.h"
 #include "Delegates/Delegate.h"
diff --git a/Engine/Plugins/VirtualProduction/CaptureManager/Source/CaptureUtils/Private/Async/EventSourceUtils.cpp b/Engine/Plugins/VirtualProduction/CaptureManager/Source/CaptureUtils/Private/Async/EventSourceUtils.cpp
--- a/Engine/Plugins/VirtualProduction/CaptureManager/Source/CaptureUtils/Private/Async/EventSourceUtils.cpp
+++ b/Engine/Plugins/VirtualProduction/CaptureManager/Source/CaptureUtils/Private/Async/EventSourceUtils.cpp
@@ -1,7 +1,5 @@
 // Copyright Epic Games, Inc. All Rights Reserved.
 
-#pragma once
-
 #include "Async/EventSourceUtils.h"
 #include "Misc/ScopeRWLock.h"
 
